Motion Design: Added motion design scene state plugin

diff --git a/Engine/Plugins/Experimental/SceneState/Source/SceneState/Private/Conduit/SceneStateConduit.cpp b/Engine/Plugins/Experimental/SceneState/Source/SceneState/Private/Conduit/SceneStateConduit.cpp
new file mode 100644
--- /dev/null
+++ b/Engine/Plugins/Experimental/SceneState/Source/SceneState/Private/Conduit/SceneStateConduit.cpp
@@ -0,0 +1,67 @@
+// Copyright Epic Games, Inc. All Rights Reserved.
+
+#include "Conduit/SceneStateConduit.h"
+#include "Conduit/SceneStateConduitLink.h"
+#include "SceneStateExecutionContext.h"
+#include "SceneStateGeneratedClass.h"
+#include "SceneStateObject.h"
+#include "Transition/SceneStateTransitionEvaluation.h"
+#include "Transition/SceneStateTransitionResult.h"
+
+void FSceneStateConduit::Link(const FSceneStateConduitLink& InConduitLink, USceneStateGeneratedClass* InGeneratedClass)
+{
+	if (!InConduitLink.ResultPropertyName.IsNone())
+	{
+		ResultProperty = CastField<FStructProperty>(InGeneratedClass->FindPropertyByName(InConduitLink.ResultPropertyName));
+		check(ResultProperty);
+	}
+	else
+	{
+		ResultProperty = nullptr;
+	}
+
+	if (!InConduitLink.EventName.IsNone())
+	{
+		EvaluationEvent = InGeneratedClass->FindFunctionByName(InConduitLink.EventName);
+		check(EvaluationEvent);
+	}
+	else
+	{
+		EvaluationEvent = nullptr;
+	}
+}
+
+bool FSceneStateConduit::Evaluate(const UE::SceneState::FTransitionEvaluationParams& InParams) const
+{
+	// Early return if waiting for tasks to finish and there are still active tasks yet to finish
+	if (EnumHasAnyFlags(EvaluationFlags, ESceneStateTransitionEvaluationFlags::WaitForTasksToFinish)
+		&& InParams.SourceState.HasPendingTasks(InParams.ExecutionContext))
+	{
+		return false;
+	}
+
+	return ProcessEvaluationEvent(InParams);
+}
+
+bool FSceneStateConduit::ProcessEvaluationEvent(const UE::SceneState::FTransitionEvaluationParams& InParams) const
+{
+	if (EnumHasAllFlags(EvaluationFlags, ESceneStateTransitionEvaluationFlags::EvaluationEventAlwaysTrue))
+	{
+		return true;
+	}
+
+	if (!EvaluationEvent || !ResultProperty)
+	{
+		return false;
+	}
+
+	void* FunctionParams = nullptr;
+
+	USceneStateObject* RootState = InParams.ExecutionContext.GetRootState();
+	check(RootState);
+	RootState->ProcessEvent(EvaluationEvent, FunctionParams);
+
+	const FSceneStateTransitionResult* Result = ResultProperty->ContainerPtrToValuePtr<FSceneStateTransitionResult>(RootState);
+	check(Result);
+	return Result->bCanTransition;
+}
diff --git a/Engine/Plugins/Experimental/SceneState/Source/SceneState/Private/SceneState.cpp b/Engine/Plugins/Experimental/SceneState/Source/SceneState/Private/SceneState.cpp
new file mode 100644
--- /dev/null
+++ b/Engine/Plugins/Experimental/SceneState/Source/SceneState/Private/SceneState.cpp
@@ -0,0 +1,311 @@
+// Copyright Epic Games, Inc. All Rights Reserved.
+
+#include "SceneState.h"
+#include "SceneStateEventHandler.h"
+#include "SceneStateEventStream.h"
+#include "SceneStateExecutionContext.h"
+#include "SceneStateInstance.h"
+#include "SceneStateLog.h"
+#include "SceneStateMachine.h"
+#include "SceneStateObject.h"
+#include "SceneStateUtils.h"
+#include "StructUtils/InstancedStructContainer.h"
+#include "Tasks/SceneStateTask.h"
+#include "Tasks/SceneStateTaskInstance.h"
+#include "Transition/SceneStateTransition.h"
+
+void FSceneState::Enter(const FSceneStateExecutionContext& InContext) const
+{
+	using namespace UE::SceneState;
+
+	const FReentryGuard ReentryGuard(ReentryHandle, InContext);
+	if (ReentryGuard.IsReentry())
+	{
+		return;
+	}
+
+	// Enter State, add a State Instance if not already present
+	FSceneStateInstance* Instance = InContext.FindOrAddStateInstance(*this);
+	if (!Instance || Instance->Status == EExecutionStatus::Running)
+	{
+		return;
+	}
+
+	UE_LOG(LogSceneState, Verbose, TEXT("State (%s) receiving enter"), GetStateName(InContext).GetData());
+
+	Instance->ElapsedTime = 0.f;
+	Instance->Status = EExecutionStatus::Running;
+
+	// Apply Event Handlers before anything else starts so the Event Data becomes available to substate machines and tasks
+	CaptureEvents(InContext);
+	AllocateTaskInstances(InContext, InContext.GetTemplateTaskInstances(*this));
+
+	for (const FSceneStateTransition& Transition : InContext.GetTransitions(*this))
+	{
+		Transition.Setup(InContext);
+	}
+
+	for (const FSceneStateMachine& StateMachine : InContext.GetStateMachines(*this))
+	{
+		StateMachine.Setup(InContext);
+	}
+
+	// Setup each Task
+	InContext.ForEachTask(*this,
+		[&InContext, Instance](const FSceneStateTask& InTask, FStructView InTaskInstance)->EIterationResult
+		{
+			InTask.Setup(InContext, InTaskInstance);
+			return EIterationResult::Continue;
+		});
+
+	for (const FSceneStateMachine& StateMachine : InContext.GetStateMachines(*this))
+	{
+		StateMachine.Start(InContext);
+	}
+
+	UpdateActiveTasks(InContext, *Instance);
+}
+
+void FSceneState::Tick(const FSceneStateExecutionContext& InContext, float InDeltaSeconds) const
+{
+	using namespace UE::SceneState;
+
+	const FReentryGuard ReentryGuard(ReentryHandle, InContext);
+	if (ReentryGuard.IsReentry())
+	{
+		return;
+	}
+
+	FSceneStateInstance* Instance = InContext.FindStateInstance(*this);
+	if (!Instance || Instance->Status != EExecutionStatus::Running)
+	{
+		return;
+	}
+
+	UE_LOG(LogSceneState, VeryVerbose, TEXT("State (%s) receiving tick"), GetStateName(InContext).GetData());
+
+	Instance->ElapsedTime += InDeltaSeconds;
+
+	UpdateActiveTasks(InContext, *Instance);
+
+	for (const FSceneStateMachine& StateMachine : InContext.GetStateMachines(*this))
+	{
+		StateMachine.Tick(InContext, InDeltaSeconds);
+	}
+
+	InContext.ForEachTask(*this,
+		[&InContext, Instance, InDeltaSeconds](const FSceneStateTask& InTask, FStructView InTaskInstance)
+		{
+			InTask.Tick(InContext, InTaskInstance, InDeltaSeconds);
+			return EIterationResult::Continue;
+		});
+}
+
+void FSceneState::Exit(const FSceneStateExecutionContext& InContext) const
+{
+	using namespace UE::SceneState;
+
+	const FReentryGuard ReentryGuard(ReentryHandle, InContext);
+	if (ReentryGuard.IsReentry())
+	{
+		return;
+	}
+
+	FSceneStateInstance* Instance = InContext.FindStateInstance(*this);
+	if (!Instance || Instance->Status != EExecutionStatus::Running)
+	{
+		return;
+	}
+
+	UE_LOG(LogSceneState, Verbose, TEXT("State (%s) receiving exit"), GetStateName(InContext).GetData());
+
+	// Stop State Machines that are still running
+	for (const FSceneStateMachine& StateMachine : InContext.GetStateMachines(*this))
+	{
+		StateMachine.Stop(InContext);
+	}
+
+	// Stop tasks that are still running
+	InContext.ForEachTask(*this,
+		[&InContext, Instance](const FSceneStateTask& InTask, FStructView InTaskInstance)
+		{
+			InTask.Stop(InContext, InTaskInstance, ESceneStateTaskStopReason::State);
+			return EIterationResult::Continue;
+		});
+
+	// Notify transitions of stop
+	for (const FSceneStateTransition& Transition : InContext.GetTransitions(*this))
+	{
+		Transition.Exit(InContext);
+	}
+
+	Instance->Status = EExecutionStatus::Finished;
+	Instance->ElapsedTime = 0.f;
+
+	InContext.RemoveStateInstance(*this);
+	InContext.RemoveTaskInstanceContainer(*this);
+
+	ResetCapturedEvents(InContext);
+}
+
+void FSceneState::UpdateActiveTasks(const FSceneStateExecutionContext& InContext, FSceneStateInstance& InInstance) const
+{
+	using namespace UE::SceneState;
+
+	FInstancedStructContainer* TaskInstanceContainer = InContext.FindTaskInstanceContainer(*this);
+	if (!TaskInstanceContainer)
+	{
+		return;
+	}
+
+	auto FindTaskInstance = 
+		[TaskInstanceContainer](uint16 InRelativeIndex)->FSceneStateTaskInstance*
+		{
+			if (TaskInstanceContainer->IsValidIndex(InRelativeIndex))
+			{
+				return (*TaskInstanceContainer)[InRelativeIndex].GetPtr<FSceneStateTaskInstance>();
+			}
+			return nullptr;
+		};
+
+	// Start all the Tasks that haven't started yet and that meet their pre-requisites
+	InContext.ForEachTask(*this,
+		[&InContext, &InInstance, &FindTaskInstance](const FSceneStateTask& InTask, FStructView InTaskInstance)->EIterationResult
+		{
+			FSceneStateTaskInstance* TaskInstance = InTaskInstance.GetPtr<FSceneStateTaskInstance>();
+			if (!TaskInstance || TaskInstance->Status != EExecutionStatus::NotStarted)
+			{
+				return EIterationResult::Continue;
+			}
+
+			bool bPrerequisitesMet = true;
+			for (uint16 Prerequisite : InContext.GetTaskPrerequisites(InTask))
+			{
+				FSceneStateTaskInstance* const PrerequisiteTaskInstance = FindTaskInstance(Prerequisite);
+				if (!PrerequisiteTaskInstance || PrerequisiteTaskInstance->Status != EExecutionStatus::Finished)
+				{
+					bPrerequisitesMet = false;
+					break;
+				}
+			}
+
+			if (bPrerequisitesMet)
+			{
+				InTask.Start(InContext, InTaskInstance);
+			}
+
+			return EIterationResult::Continue;
+		});
+}
+
+FStringView FSceneState::GetStateName(const FSceneStateExecutionContext& InContext) const
+{
+#if WITH_EDITOR
+	if (const FSceneStateMetadata* StateMetadata = InContext.GetStateMetadata(*this))
+	{
+		return StateMetadata->StateName;
+	}
+#endif
+	return FStringView();
+}
+
+bool FSceneState::HasPendingTasks(const FSceneStateExecutionContext& InContext) const
+{
+	using namespace UE::SceneState;
+
+	FSceneStateInstance* const StateInstance = InContext.FindStateInstance(*this);
+	if (!StateInstance)
+	{
+		return false;
+	}
+
+	bool bHasPendingTask = false;
+
+	InContext.ForEachTask(*this,
+		[StateInstance, &bHasPendingTask](const FSceneStateTask& InTask, FStructView InTaskInstance)->EIterationResult
+		{
+			FSceneStateTaskInstance* TaskInstance = InTaskInstance.GetPtr<FSceneStateTaskInstance>();
+			if (TaskInstance && TaskInstance->Status != EExecutionStatus::Finished)
+			{
+				bHasPendingTask = true;
+				return EIterationResult::Break;
+			}
+			return EIterationResult::Continue;
+		});
+
+	return bHasPendingTask;
+}
+
+void FSceneState::AllocateTaskInstances(const FSceneStateExecutionContext& InContext, TConstArrayView<FConstStructView> InTemplateTaskInstances) const
+{
+	if (FInstancedStructContainer* const TaskInstanceContainer = InContext.FindOrAddTaskInstanceContainer(*this))
+	{
+		// Copy the Template data
+		*TaskInstanceContainer = InTemplateTaskInstances;
+
+		// Instance each Template Object into the Instance data
+		InstanceTaskObjects(InContext.GetRootState()
+			, UE::SceneState::GetStructViews(*TaskInstanceContainer)
+			, InTemplateTaskInstances
+			, [](FObjectDuplicationParameters& InParams)->UObject*
+			{
+				return StaticDuplicateObjectEx(InParams);
+			});
+	}
+}
+
+void FSceneState::InstanceTaskObjects(UObject* InOuter, TConstArrayView<FStructView> InTargets, TConstArrayView<FConstStructView> InSources, TFunctionRef<UObject*(FObjectDuplicationParameters&)> InDuplicationFunc) const
+{
+	check(InTargets.Num() == InSources.Num());
+
+	for (int32 TaskIndex = 0; TaskIndex < InSources.Num(); ++TaskIndex)
+	{
+		const FConstStructView Source = InSources[TaskIndex];
+		const FStructView Target = InTargets[TaskIndex];
+
+		check(Source.GetScriptStruct() == Target.GetScriptStruct());
+
+		if (!Source.GetScriptStruct())
+		{
+			continue;
+		}
+
+		for (FObjectProperty* ObjectProperty : TFieldRange<FObjectProperty>(Source.GetScriptStruct()))
+		{
+			if (!ObjectProperty->HasAnyPropertyFlags(CPF_InstancedReference))
+			{
+				continue;
+			}
+
+			UObject* InstanceObject = nullptr;
+			if (UObject* TemplateObject = ObjectProperty->GetObjectPropertyValue_InContainer(Source.GetMemory()))
+			{
+				FObjectDuplicationParameters Parameters = InitStaticDuplicateObjectParams(TemplateObject, InOuter);
+				Parameters.DestName = MakeUniqueObjectName(Parameters.DestOuter, Parameters.SourceObject->GetClass(), Parameters.SourceObject->GetFName());
+				Parameters.FlagMask = RF_AllFlags & ~RF_DefaultSubObject;
+				Parameters.PortFlags |= PPF_DuplicateVerbatim; // Skip resetting text IDs
+
+				InstanceObject = InDuplicationFunc(Parameters);
+			}
+			ObjectProperty->SetPropertyValue_InContainer(Target.GetMemory(), InstanceObject);
+		}
+	}
+}
+
+void FSceneState::CaptureEvents(const FSceneStateExecutionContext& InContext) const
+{
+	if (USceneStateEventStream* EventStream = InContext.GetEventStream())
+	{
+		TConstArrayView<FSceneStateEventHandler> EventHandlers = InContext.GetEventHandlers(*this);
+		EventStream->CaptureEvents(EventHandlers);
+	}
+}
+
+void FSceneState::ResetCapturedEvents(const FSceneStateExecutionContext& InContext) const
+{
+	if (USceneStateEventStream* EventStream = InContext.GetEventStream())
+	{
+		TConstArrayView<FSceneStateEventHandler> EventHandlers = InContext.GetEventHandlers(*this);
+		EventStream->ResetCapturedEvents(EventHandlers);
+	}
+}
diff --git a/Engine/Plugins/Experimental/SceneState/Source/SceneState/Private/SceneStateExecutionContext.cpp b/Engine/Plugins/Experimental/SceneState/Source/SceneState/Private/SceneStateExecutionContext.cpp
new file mode 100644
--- /dev/null
+++ b/Engine/Plugins/Experimental/SceneState/Source/SceneState/Private/SceneStateExecutionContext.cpp
@@ -0,0 +1,685 @@
+// Copyright Epic Games, Inc. All Rights Reserved.
+
+#include "SceneStateExecutionContext.h"
+#include "PropertyBindingDataView.h"
+#include "PropertyBindingTypes.h"
+#include "SceneState.h"
+#include "SceneStateExecutionContextRegistry.h"
+#include "SceneStateGeneratedClass.h"
+#include "SceneStateLog.h"
+#include "SceneStateMachine.h"
+#include "SceneStateObject.h"
+#include "SceneStateUtils.h"
+#include "Tasks/SceneStateTaskBindingExtension.h"
+#include "Tasks/SceneStateTaskInstance.h"
+
+const FSceneStateExecutionContext FSceneStateExecutionContext::InvalidContext;
+
+FSceneStateExecutionContext::~FSceneStateExecutionContext()
+{
+	Reset();
+}
+
+TWeakPtr<const UE::SceneState::FExecutionContextRegistry> FSceneStateExecutionContext::GetContextRegistry() const
+{
+	return ContextRegistryWeak;
+}
+
+const FSceneStateBindingCollection* FSceneStateExecutionContext::GetBindingCollection() const
+{
+	if (GeneratedClass)
+	{
+		return &GeneratedClass->GetBindingCollection();
+	}
+	return nullptr;
+}
+
+FString FSceneStateExecutionContext::GetExecutionContextName() const
+{
+	if (RootState)
+	{
+		return RootState->GetContextName();
+	}
+	return FString();
+}
+
+UObject* FSceneStateExecutionContext::GetContextObject() const
+{
+	if (RootState)
+	{
+		return RootState->GetContextObject();
+	}
+	return nullptr;
+}
+
+USceneStateEventStream* FSceneStateExecutionContext::GetEventStream() const
+{
+	if (RootState)
+	{
+		return RootState->GetEventStream();
+	}
+	return nullptr;
+}
+
+void FSceneStateExecutionContext::Setup(USceneStateObject* InRootState)
+{
+	Reset();
+
+	check(InRootState);
+	RootState = InRootState;
+	GeneratedClass = Cast<USceneStateGeneratedClass>(RootState->GetClass());
+
+	const TSharedRef<UE::SceneState::FExecutionContextRegistry> ContextRegistry = InRootState->GetContextRegistry();
+	ContextRegistry->RegisterContext(this);
+	ContextRegistryWeak = ContextRegistry;
+}
+
+void FSceneStateExecutionContext::Reset()
+{
+	RootState = nullptr;
+	GeneratedClass = nullptr;
+
+	StateInstances.Reset();
+	StateMachineInstances.Reset();
+	TransitionInstances.Reset();
+
+	if (TSharedPtr<UE::SceneState::FExecutionContextRegistry> ContextRegistry = ContextRegistryWeak.Pin())
+	{
+		ContextRegistry->UnregisterContext(this);
+	}
+
+	ContextRegistryWeak.Reset();
+}
+
+FPropertyBindingDataView FSceneStateExecutionContext::FindDataView(const FSceneStateBindingDataHandle& InDataHandle) const
+{
+	const uint16 DataIndex = InDataHandle.GetDataIndex();
+	const uint16 DataSubIndex = InDataHandle.GetDataSubIndex();
+
+	if (InDataHandle.IsExternalDataType())
+	{
+		// Todo: unsupported external sources
+		return FPropertyBindingDataView();
+	}
+
+	const ESceneStateDataType DataType = static_cast<ESceneStateDataType>(InDataHandle.GetDataType());
+
+	switch (DataType)
+	{
+	case ESceneStateDataType::Root:
+		return FPropertyBindingDataView(RootState);
+
+	case ESceneStateDataType::Task:
+		return FindTaskInstance(DataIndex);
+
+	case ESceneStateDataType::EventHandler:
+		if (USceneStateEventStream* EventStream = GetEventStream())
+		{
+			if (const FSceneStateEventHandler* EventHandler = GetEventHandler(DataIndex))
+			{
+				if (FSceneStateEvent* Event = EventStream->FindCapturedEvent(EventHandler->GetHandlerId()))
+				{
+					return Event->GetDataViewMutable();
+				}
+			}
+		}
+		break;
+
+	case ESceneStateDataType::TaskExtension:
+		if (const FSceneStateTask* Task = FindTask(DataIndex).GetPtr<const FSceneStateTask>())
+		{
+			if (const FSceneStateTaskBindingExtension* BindingExtension = Task->GetBindingExtension())
+			{
+				FStructView DataView;
+				if (BindingExtension->FindDataByIndex(FindTaskInstance(DataIndex), DataSubIndex, DataView))
+				{
+					return DataView;
+				}
+			}
+		}
+		break;
+
+	case ESceneStateDataType::Transition:
+		if (FSceneStateTransitionInstance* TransitionInstance = FindTransitionInstance(DataIndex))
+		{
+			return TransitionInstance->Parameters.GetMutableValue();
+		}
+		break;
+
+	case ESceneStateDataType::StateMachine:
+		if (FSceneStateMachineInstance* StateMachineInstance = FindStateMachineInstance(DataIndex))
+		{
+			return StateMachineInstance->Parameters.GetMutableValue();
+		}
+		break;
+	}
+
+	return FPropertyBindingDataView();
+}
+
+void FSceneStateExecutionContext::ForEachTask(const FSceneState& InState, TFunctionRef<UE::SceneState::EIterationResult(const FSceneStateTask&, FStructView)> InCallable) const
+{
+	const FSceneStateRange TaskRange = InState.GetTaskRange();
+	if (TaskRange.Count == 0 || !GeneratedClass)
+	{
+		return;
+	}
+
+	FInstancedStructContainer* TaskInstanceContainer = FindTaskInstanceContainer(InState);
+	if (!TaskInstanceContainer)
+	{
+		return;
+	}
+
+	const FInstancedStructContainer& Tasks = GeneratedClass->GetTasks();
+
+	checkfSlow(Tasks.IsValidIndex(TaskRange.Index) && Tasks.IsValidIndex(TaskRange.GetLastIndex())
+		, TEXT("[Context: %s] State Task Range [%d, %d] is out of bounds! (Total Task Count: %d)")
+		, *GetExecutionContextName()
+		, TaskRange.Index
+		, TaskRange.GetLastIndex()
+		, Tasks.Num());
+
+	for (uint16 TaskIndex = 0; TaskIndex < TaskRange.Count; ++TaskIndex)
+	{
+		if (!TaskInstanceContainer->IsValidIndex(TaskIndex))
+		{
+			continue;
+		}
+
+		const uint16 AbsoluteIndex = TaskRange.Index + TaskIndex;
+		const FSceneStateTask& Task = Tasks[AbsoluteIndex].Get<const FSceneStateTask>();
+		const FStructView TaskInstance = (*TaskInstanceContainer)[TaskIndex];
+
+		if (InCallable(Task, TaskInstance) == UE::SceneState::EIterationResult::Break)
+		{
+			break;
+		}
+	}
+}
+
+TArray<FConstStructView> FSceneStateExecutionContext::GetTemplateTaskInstances(const FSceneState& InState) const
+{
+	const FSceneStateRange TaskRange = InState.GetTaskRange();
+	if (TaskRange.Count == 0 || !GeneratedClass)
+	{
+		return {};
+	}
+	return UE::SceneState::GetConstStructViews(GeneratedClass->GetTaskInstances(), TaskRange);
+}
+
+TConstArrayView<uint16> FSceneStateExecutionContext::GetTaskPrerequisites(const FSceneStateTask& InTask) const
+{
+	const FSceneStateRange TaskPrerequisiteRange = InTask.GetPrerequisiteRange();
+	if (TaskPrerequisiteRange.Count == 0 || !GeneratedClass)
+	{
+		return {};
+	}
+
+	const TConstArrayView<uint16> TaskPrerequisitesView = GeneratedClass->GetTaskPrerequisites();
+
+	checkfSlow(TaskPrerequisitesView.IsValidIndex(TaskPrerequisiteRange.Index) && TaskPrerequisitesView.IsValidIndex(TaskPrerequisiteRange.GetLastIndex())
+		, TEXT("[Context: %s] Task Prerequisite Range [%d, %d] is out of bounds! (Total Prerequisite Count: %d)")
+		, *GetExecutionContextName()
+		, TaskPrerequisiteRange.Index
+		, TaskPrerequisiteRange.GetLastIndex()
+		, TaskPrerequisitesView.Num());
+
+	return TaskPrerequisitesView.Slice(TaskPrerequisiteRange.Index, TaskPrerequisiteRange.Count);
+}
+
+TConstArrayView<FSceneStateTransition> FSceneStateExecutionContext::GetTransitions(const FSceneState& InState) const
+{
+	const FSceneStateRange TransitionRange = InState.GetTransitionRange();
+	if (TransitionRange.Count == 0 || !GeneratedClass)
+	{
+		return {};
+	}
+
+	const TConstArrayView<FSceneStateTransition> TransitionsView = GeneratedClass->GetTransitions();
+
+	checkfSlow(TransitionsView.IsValidIndex(TransitionRange.Index) && TransitionsView.IsValidIndex(TransitionRange.GetLastIndex())
+		, TEXT("[Context: %s] State Transition Range [%d, %d] is out of bounds! (Total Transition Count: %d)")
+		, *GetExecutionContextName()
+		, TransitionRange.Index
+		, TransitionRange.GetLastIndex()
+		, TransitionsView.Num());
+
+	return TransitionsView.Slice(TransitionRange.Index, TransitionRange.Count);
+}
+
+TConstArrayView<FSceneStateTransition> FSceneStateExecutionContext::GetTransitions(const FSceneStateConduit& InConduit) const
+{
+	const FSceneStateRange TransitionRange = InConduit.GetTransitionRange();
+	if (TransitionRange.Count == 0 || !GeneratedClass)
+	{
+		return {};
+	}
+
+	const TConstArrayView<FSceneStateTransition> TransitionsView = GeneratedClass->GetTransitions();
+
+	checkfSlow(TransitionsView.IsValidIndex(TransitionRange.Index) && TransitionsView.IsValidIndex(TransitionRange.GetLastIndex())
+		, TEXT("[Context: %s] Conduit Transition Range [%d, %d] is out of bounds! (Total Transition Count: %d)")
+		, *GetExecutionContextName()
+		, TransitionRange.Index
+		, TransitionRange.GetLastIndex()
+		, TransitionsView.Num());
+
+	return TransitionsView.Slice(TransitionRange.Index, TransitionRange.Count);
+}
+
+FInstancedPropertyBag FSceneStateExecutionContext::GetTemplateTransitionParameter(const FSceneStateTransition& InTransition) const
+{
+	uint16 TransitionIndex;
+	if (GetTransitionIndex(InTransition, TransitionIndex))
+	{
+		return GeneratedClass->GetTransitionParameter(TransitionIndex);
+	}
+	return FInstancedPropertyBag();
+}
+
+TConstArrayView<FSceneStateMachine> FSceneStateExecutionContext::GetStateMachines(const FSceneState& InState) const
+{
+	const FSceneStateRange StateMachineRange = InState.GetStateMachineRange();
+	if (StateMachineRange.Count == 0 || !GeneratedClass)
+	{
+		return {};
+	}
+
+	const TConstArrayView<FSceneStateMachine> StateMachinesView = GeneratedClass->GetStateMachines();
+
+	checkfSlow(StateMachinesView.IsValidIndex(StateMachineRange.Index) && StateMachinesView.IsValidIndex(StateMachineRange.GetLastIndex())
+		, TEXT("[Context: %s] State State Machine Range [%d, %d] is out of bounds! (Total State Machine Count: %d)")
+		, *GetExecutionContextName()
+		, StateMachineRange.Index
+		, StateMachineRange.GetLastIndex()
+		, StateMachinesView.Num());
+
+	return StateMachinesView.Slice(StateMachineRange.Index, StateMachineRange.Count);
+}
+
+TConstArrayView<FSceneStateEventHandler> FSceneStateExecutionContext::GetEventHandlers(const FSceneState& InState) const
+{
+	const FSceneStateRange EventHandlerRange = InState.GetEventHandlerRange();
+	if (EventHandlerRange.Count == 0 || !GeneratedClass)
+	{
+		return {};
+	}
+
+	const TConstArrayView<FSceneStateEventHandler> EventHandlersView = GeneratedClass->GetEventHandlers();
+
+	checkfSlow(EventHandlersView.IsValidIndex(EventHandlerRange.Index) && EventHandlersView.IsValidIndex(EventHandlerRange.GetLastIndex())
+		, TEXT("[Context: %s] State Event Handler Range [%d, %d] is out of bounds! (Total Event Handler Count: %d)")
+		, *GetExecutionContextName()
+		, EventHandlerRange.Index
+		, EventHandlerRange.GetLastIndex()
+		, EventHandlersView.Num());
+
+	return EventHandlersView.Slice(EventHandlerRange.Index, EventHandlerRange.Count);
+}
+
+#if WITH_EDITOR
+const FSceneStateMetadata* FSceneStateExecutionContext::GetStateMetadata(const FSceneState& InState) const
+{
+	if (!GeneratedClass)
+	{
+		return nullptr;
+	}
+
+	uint16 StateIndex;
+	if (!GetStateIndex(InState, StateIndex))
+	{
+		UE_LOG(LogSceneState, Error, TEXT("[Context: %s] State Index was not able to be retrieved. State Data is invalid!"), *GetExecutionContextName());
+		return nullptr;
+	}
+
+	const TConstArrayView<FSceneStateMetadata> StateMetadataView = GeneratedClass->GetStateMetadata();
+
+	if (!StateMetadataView.IsValidIndex(StateIndex))
+	{
+		UE_LOG(LogSceneState, Error, TEXT("[Context: %s] Metadata for State Index %d is not available!"
+			"Be sure to have compiled the Scene State Object. Metadata count: %d")
+			, *GetExecutionContextName()
+			, StateIndex
+			, StateMetadataView.Num());
+		return nullptr;
+	}
+
+	return &StateMetadataView[StateIndex];
+}
+#endif
+
+const FSceneStateMachine* FSceneStateExecutionContext::GetStateMachine(const FGuid& InStateMachineId) const
+{
+	if (GeneratedClass)
+	{
+		return GeneratedClass->FindStateMachine(InStateMachineId);
+	}
+	return nullptr;
+}
+
+const FSceneState* FSceneStateExecutionContext::GetActiveState(const FSceneStateMachine& InStateMachine) const
+{
+	if (const FSceneStateMachineInstance* Instance = FindStateMachineInstance(InStateMachine))
+	{
+		return GetState(InStateMachine, Instance->ActiveIndex);
+	}
+	return nullptr;
+}
+
+const FSceneState* FSceneStateExecutionContext::GetState(uint16 InAbsoluteIndex) const
+{
+	if (!GeneratedClass)
+	{
+		return nullptr;
+	}
+
+	const TConstArrayView<FSceneState> StatesView = GeneratedClass->GetStates();
+
+	checkfSlow(StatesView.IsValidIndex(InAbsoluteIndex)
+		, TEXT("[Context: %s] Absolute Index %d is out of bounds! (Total State Count: %d)")
+		, *GetExecutionContextName()
+		, InAbsoluteIndex
+		, StatesView.Num());
+
+	return &StatesView[InAbsoluteIndex];
+}
+
+const FSceneStateEventHandler* FSceneStateExecutionContext::GetEventHandler(uint16 InAbsoluteIndex) const
+{
+	if (!GeneratedClass)
+	{
+		return nullptr;
+	}
+
+	const TConstArrayView<FSceneStateEventHandler> EventHandlers = GeneratedClass->GetEventHandlers();
+
+	checkfSlow(EventHandlers.IsValidIndex(InAbsoluteIndex)
+		, TEXT("[Context: %s] Absolute Index %d is out of bounds! (Total Event Handlers Count: %d)")
+		, *GetExecutionContextName()
+		, InAbsoluteIndex
+		, EventHandlers.Num());
+
+	return &EventHandlers[InAbsoluteIndex];
+}
+
+const FSceneState* FSceneStateExecutionContext::GetState(const FSceneStateMachine& InStateMachine, uint16 InRelativeIndex) const
+{
+	const FSceneStateRange StateRange = InStateMachine.GetStateRange();
+	if (StateRange.Count == 0 || !GeneratedClass)
+	{
+		return nullptr;
+	}
+
+	const TConstArrayView<FSceneState> StatesView = GeneratedClass->GetStates();
+
+	const uint16 AbsoluteIndex = StateRange.Index + InRelativeIndex;
+
+	checkfSlow(StatesView.IsValidIndex(AbsoluteIndex)
+		, TEXT("[Context: %s] State Machine Relative Active Index %d from [%d, %d] is out of bounds! (Total State Count: %d)")
+		, *GetExecutionContextName()
+		, InRelativeIndex
+		, StateRange.Index
+		, StateRange.GetLastIndex()
+		, StatesView.Num());
+
+	return &StatesView[AbsoluteIndex];
+}
+
+const FSceneStateConduit* FSceneStateExecutionContext::GetConduit(const FSceneStateMachine& InStateMachine, uint16 InRelativeIndex) const
+{
+	const FSceneStateRange ConduitRange = InStateMachine.GetConduitRange();
+	if (ConduitRange.Count == 0 || !GeneratedClass)
+	{
+		return nullptr;
+	}
+
+	const TConstArrayView<FSceneStateConduit> ConduitsView = GeneratedClass->GetConduits();
+
+	const uint16 AbsoluteIndex = ConduitRange.Index + InRelativeIndex;
+
+	checkfSlow(ConduitsView.IsValidIndex(AbsoluteIndex)
+		, TEXT("[Context: %s] State machine conduit relative index %d from [%d, %d] is out of bounds! (Total conduit count: %d)")
+		, *GetExecutionContextName()
+		, InRelativeIndex
+		, ConduitRange.Index
+		, ConduitRange.GetLastIndex()
+		, ConduitsView.Num());
+
+	return &ConduitsView[AbsoluteIndex];
+}
+
+FSceneStateInstance* FSceneStateExecutionContext::FindOrAddStateInstance(const FSceneState& InState) const
+{
+	uint16 StateIndex;
+	if (GetStateIndex(InState, /*out*/StateIndex))
+	{
+		return &StateInstances.FindOrAdd(StateIndex);
+	}
+	return nullptr;
+}
+
+FSceneStateInstance* FSceneStateExecutionContext::FindStateInstance(uint16 InAbsoluteIndex) const
+{
+	return StateInstances.Find(InAbsoluteIndex);
+}
+
+FSceneStateInstance* FSceneStateExecutionContext::FindStateInstance(const FSceneState& InState) const
+{
+	uint16 StateIndex;
+	if (GetStateIndex(InState, /*out*/StateIndex))
+	{
+		return FindStateInstance(StateIndex);
+	}
+	return nullptr;
+}
+
+void FSceneStateExecutionContext::RemoveStateInstance(const FSceneState& InState) const
+{
+	uint16 StateIndex;
+	if (GetStateIndex(InState, /*out*/StateIndex))
+	{
+		StateInstances.Remove(StateIndex);
+	}
+}
+
+FConstStructView FSceneStateExecutionContext::FindTask(uint16 InAbsoluteIndex) const
+{
+	if (!GeneratedClass)
+	{
+		return {};
+	}
+
+	const FInstancedStructContainer& Tasks = GeneratedClass->GetTasks();
+
+	if (!Tasks.IsValidIndex(InAbsoluteIndex))
+	{
+		return {};
+	}
+
+	return Tasks[InAbsoluteIndex];
+}
+
+FInstancedStructContainer* FSceneStateExecutionContext::FindOrAddTaskInstanceContainer(const FSceneState& InState) const
+{
+	uint16 StateIndex;
+	if (GetStateIndex(InState, /*out*/StateIndex))
+	{
+		return &TaskInstanceContainers.FindOrAdd(StateIndex);
+	}
+	return nullptr;
+}
+
+FInstancedStructContainer* FSceneStateExecutionContext::FindTaskInstanceContainer(const FSceneState& InState) const
+{
+	uint16 StateIndex;
+	if (GetStateIndex(InState, /*out*/StateIndex))
+	{
+		return FindTaskInstanceContainer(StateIndex);
+	}
+	return nullptr;
+}
+
+FInstancedStructContainer* FSceneStateExecutionContext::FindTaskInstanceContainer(uint16 InAbsoluteIndex) const
+{
+	return TaskInstanceContainers.Find(InAbsoluteIndex);
+}
+
+FStructView FSceneStateExecutionContext::FindTaskInstance(uint16 InAbsoluteIndex) const
+{
+	if (!GeneratedClass)
+	{
+		return {};
+	}
+
+	const FInstancedStructContainer& Tasks = GeneratedClass->GetTasks();
+
+	if (!Tasks.IsValidIndex(InAbsoluteIndex))
+	{
+		return {};
+	}
+
+	const FSceneStateTask& Task = Tasks[InAbsoluteIndex].Get<const FSceneStateTask>();
+
+	const uint16 StateIndex = Task.GetParentStateIndex();
+
+	FInstancedStructContainer* const TaskInstanceContainer = FindTaskInstanceContainer(StateIndex);
+	if (!TaskInstanceContainer)
+	{
+		return {};
+	}
+
+	const TConstArrayView<FSceneState> StatesView = GeneratedClass->GetStates();
+
+	checkf(StatesView.IsValidIndex(StateIndex)
+		, TEXT("[Context: %s] Task Parent State Index %d is out of bounds! (Total State Count: %d)")
+		, *GetExecutionContextName()
+		, StateIndex
+		, StatesView.Num());
+
+	const FSceneState& State = StatesView[StateIndex];
+
+	const uint16 TaskRelativeIndex = InAbsoluteIndex - State.GetTaskRange().Index;
+	if (!TaskInstanceContainer->IsValidIndex(TaskRelativeIndex))
+	{
+		return {};
+	}
+
+	return (*TaskInstanceContainer)[TaskRelativeIndex];
+}
+
+void FSceneStateExecutionContext::RemoveTaskInstanceContainer(const FSceneState& InState) const
+{
+	uint16 StateIndex;
+	if (GetStateIndex(InState, /*out*/StateIndex))
+	{
+		TaskInstanceContainers.Remove(StateIndex);
+	}
+}
+
+FSceneStateMachineInstance* FSceneStateExecutionContext::FindOrAddStateMachineInstance(const FSceneStateMachine& InStateMachine) const
+{
+	uint16 StateMachineIndex;
+	if (GetStateMachineIndex(InStateMachine, /*out*/StateMachineIndex))
+	{
+		return &StateMachineInstances.FindOrAdd(StateMachineIndex);
+	}
+	return nullptr;
+}
+
+FSceneStateMachineInstance* FSceneStateExecutionContext::FindStateMachineInstance(uint16 InAbsoluteIndex) const
+{
+	return StateMachineInstances.Find(InAbsoluteIndex);
+}
+
+FSceneStateMachineInstance* FSceneStateExecutionContext::FindStateMachineInstance(const FSceneStateMachine& InStateMachine) const
+{
+	uint16 StateMachineIndex;
+	if (GetStateMachineIndex(InStateMachine, /*out*/StateMachineIndex))
+	{
+		return FindStateMachineInstance(StateMachineIndex);
+	}
+	return nullptr;
+}
+
+void FSceneStateExecutionContext::RemoveStateMachineInstance(const FSceneStateMachine& InStateMachine) const
+{
+	uint16 StateMachineIndex;
+	if (GetStateMachineIndex(InStateMachine, /*out*/StateMachineIndex))
+	{
+		StateMachineInstances.Remove(StateMachineIndex);
+	}
+}
+
+FSceneStateTransitionInstance* FSceneStateExecutionContext::FindOrAddTransitionInstance(const FSceneStateTransition& InTransition) const
+{
+	uint16 TransitionIndex;
+	if (GetTransitionIndex(InTransition, /*out*/TransitionIndex))
+	{
+		return &TransitionInstances.FindOrAdd(TransitionIndex);
+	}
+	return nullptr;
+}
+
+FSceneStateTransitionInstance* FSceneStateExecutionContext::FindTransitionInstance(uint16 InAbsoluteIndex) const
+{
+	return TransitionInstances.Find(InAbsoluteIndex);
+}
+
+FSceneStateTransitionInstance* FSceneStateExecutionContext::FindTransitionInstance(const FSceneStateTransition& InTransition) const
+{
+	uint16 TransitionIndex;
+	if (GetTransitionIndex(InTransition, /*out*/TransitionIndex))
+	{
+		return FindTransitionInstance(TransitionIndex);
+	}
+	return nullptr;
+}
+
+void FSceneStateExecutionContext::RemoveTransitionInstance(const FSceneStateTransition& InTransition) const
+{
+	uint16 TransitionIndex;
+	if (GetTransitionIndex(InTransition, /*out*/TransitionIndex))
+	{
+		TransitionInstances.Remove(TransitionIndex);
+	}
+}
+
+bool FSceneStateExecutionContext::GetStateIndex(const FSceneState& InState, uint16& OutIndex) const
+{
+	if (GeneratedClass)
+	{
+		TConstArrayView<FSceneState> StatesView = GeneratedClass->GetStates();
+		const int32 Index = static_cast<int32>(&InState - StatesView.GetData());
+		checkSlow(StatesView.IsValidIndex(Index) && Index >= 0 && Index <= TNumericLimits<uint16>::Max());
+		OutIndex = static_cast<uint16>(Index);
+		return true;
+	}
+	return false;
+}
+
+bool FSceneStateExecutionContext::GetStateMachineIndex(const FSceneStateMachine& InStateMachine, uint16& OutIndex) const
+{
+	if (GeneratedClass)
+	{
+		TConstArrayView<FSceneStateMachine> StateMachinesView = GeneratedClass->GetStateMachines();
+		const int32 Index = static_cast<int32>(&InStateMachine - StateMachinesView.GetData());
+		checkSlow(StateMachinesView.IsValidIndex(Index) && Index >= 0 && Index <= TNumericLimits<uint16>::Max());
+		OutIndex = static_cast<uint16>(Index);
+		return true;
+	}
+	return false;
+}
+
+bool FSceneStateExecutionContext::GetTransitionIndex(const FSceneStateTransition& InTransition, uint16& OutIndex) const
+{
+	if (GeneratedClass)
+	{
+		TConstArrayView<FSceneStateTransition> TransitionsView = GeneratedClass->GetTransitions();
+		const int32 Index = static_cast<int32>(&InTransition - TransitionsView.GetData());
+		checkSlow(TransitionsView.IsValidIndex(Index) && Index >= 0 && Index <= TNumericLimits<uint16>::Max());
+		OutIndex = static_cast<uint16>(Index);
+		return true;
+	}
+	return false;
+}
diff --git a/Engine/Plugins/Experimental/SceneState/Source/SceneState/Private/SceneStateGeneratedClass.cpp b/Engine/Plugins/Experimental/SceneState/Source/SceneState/Private/SceneStateGeneratedClass.cpp
new file mode 100644
--- /dev/null
+++ b/Engine/Plugins/Experimental/SceneState/Source/SceneState/Private/SceneStateGeneratedClass.cpp
@@ -0,0 +1,298 @@
+// Copyright Epic Games, Inc. All Rights Reserved.
+
+#include "SceneStateGeneratedClass.h"
+#include "Misc/EnumerateRange.h"
+#include "SceneStateBindingUtils.h"
+#include "SceneStateEventSchema.h"
+#include "SceneStateExecutionContext.h"
+#include "SceneStateExecutionContextRegistry.h"
+#include "SceneStateLog.h"
+#include "SceneStateMachine.h"
+#include "SceneStateObject.h"
+#include "SceneStateUtils.h"
+#include "StructUtils/UserDefinedStruct.h"
+#include "Tasks/SceneStateTask.h"
+#include "Tasks/SceneStateTaskBindingExtension.h"
+#include "Tasks/SceneStateTaskInstance.h"
+
+#if WITH_EDITOR
+#include "StructUtils/UserDefinedStruct.h"
+#include "StructUtilsDelegates.h"
+#endif
+
+USceneStateGeneratedClass::USceneStateGeneratedClass()
+{
+#if WITH_EDITOR
+	if (!HasAnyFlags(RF_ClassDefaultObject))
+	{
+		OnObjectsReinstancedHandle = FCoreUObjectDelegates::OnObjectsReinstanced.AddUObject(this, &USceneStateGeneratedClass::OnObjectsReinstanced);
+		OnStructsReinstancedHandle = UE::StructUtils::Delegates::OnUserDefinedStructReinstanced.AddUObject(this, &USceneStateGeneratedClass::OnStructsReinstanced);
+	}
+#endif
+}
+
+const FSceneState* USceneStateGeneratedClass::GetRootState() const
+{
+	if (States.IsValidIndex(RootStateIndex))
+	{
+		return &States[RootStateIndex];	
+	}
+	return nullptr;
+}
+
+const FSceneStateMachine* USceneStateGeneratedClass::FindStateMachine(const FGuid& InStateMachineId) const
+{
+	const uint16* StateMachineIndex = StateMachineIdToIndex.Find(InStateMachineId);
+	if (StateMachineIndex && StateMachines.IsValidIndex(*StateMachineIndex))
+	{
+		return &StateMachines[*StateMachineIndex];
+	}
+	return nullptr;
+}
+
+FInstancedPropertyBag USceneStateGeneratedClass::GetTransitionParameter(uint16 InAbsoluteIndex) const
+{
+	return TransitionParameters.FindRef(InAbsoluteIndex);
+}
+
+void USceneStateGeneratedClass::Link(FArchive& Ar, bool bInRelinkExistingProperties)
+{
+	Super::Link(Ar, bInRelinkExistingProperties);
+
+	if (IsFullClass())
+	{
+		for (TEnumerateRef<FSceneStateTransition> Transition : EnumerateRange(Transitions))
+		{
+			if (TransitionLinks.IsValidIndex(Transition.GetIndex()))
+			{
+				Transition->Link(TransitionLinks[Transition.GetIndex()], this);
+			}
+		}
+
+		for (TEnumerateRef<FSceneStateConduit> Conduit : EnumerateRange(Conduits))
+		{
+			if (ConduitLinks.IsValidIndex(Conduit.GetIndex()))
+			{
+				Conduit->Link(ConduitLinks[Conduit.GetIndex()], this);
+			}
+		}
+	}
+}
+
+void USceneStateGeneratedClass::PostLoad()
+{
+	Super::PostLoad();
+	ResolveBindings();
+}
+
+void USceneStateGeneratedClass::BeginDestroy()
+{
+	Super::BeginDestroy();
+	Reset();
+}
+
+#if WITH_EDITOR
+const FSceneState* USceneStateGeneratedClass::FindStateFromNode(FObjectKey InStateNode) const
+{
+	const uint16* StateIndex = StateNodeToIndex.Find(InStateNode);
+	if (StateIndex && States.IsValidIndex(*StateIndex))
+	{
+		return &States[*StateIndex];
+	}
+	return nullptr;
+}
+
+const FSceneStateTask* USceneStateGeneratedClass::FindTaskFromNode(FObjectKey InTaskNode) const
+{
+	const uint16* TaskIndex = TaskNodeToIndex.Find(InTaskNode);
+	if (TaskIndex && Tasks.IsValidIndex(*TaskIndex))
+	{
+		return Tasks[*TaskIndex].GetPtr<const FSceneStateTask>();
+	}
+	return nullptr;
+}
+
+void USceneStateGeneratedClass::ForEachStateInstance(USceneStateObject& InRootState, FObjectKey InStateNode, TFunctionRef<void(const FSceneStateInstance&)> InFunctor) const
+{
+	const uint16* StateIndex = StateNodeToIndex.Find(InStateNode);
+	if (!StateIndex)
+	{
+		return;
+	}
+
+	const TSharedRef<UE::SceneState::FExecutionContextRegistry> ContextRegistry = InRootState.GetContextRegistry();
+
+	ContextRegistry->ForEachExecutionContext(
+		[StateIndex, &InFunctor](const FSceneStateExecutionContext& InExecutionContext)
+		{
+			if (const FSceneStateInstance* StateInstance = InExecutionContext.FindStateInstance(*StateIndex))
+			{
+				InFunctor(*StateInstance);
+			}
+		});
+}
+
+void USceneStateGeneratedClass::ForEachTaskInstance(USceneStateObject& InRootState, FObjectKey InTaskNode, TFunctionRef<void(const FSceneStateTaskInstance&)> InFunctor) const
+{
+	const uint16* TaskIndex = TaskNodeToIndex.Find(InTaskNode);
+	if (!TaskIndex)
+	{
+		return;
+	}
+
+	const TSharedRef<UE::SceneState::FExecutionContextRegistry> ContextRegistry = InRootState.GetContextRegistry();
+
+	ContextRegistry->ForEachExecutionContext(
+		[TaskIndex, &InFunctor](const FSceneStateExecutionContext& InExecutionContext)
+		{
+			if (const FSceneStateTaskInstance* TaskInstance = InExecutionContext.FindTaskInstance(*TaskIndex).GetPtr<const FSceneStateTaskInstance>())
+			{
+				InFunctor(*TaskInstance);
+			}
+		});
+}
+#endif
+
+const UStruct* USceneStateGeneratedClass::FindDataStruct(const FSceneStateBindingDataHandle& InDataHandle)
+{
+	const uint16 DataIndex = InDataHandle.GetDataIndex();
+	const uint16 DataSubIndex = InDataHandle.GetDataSubIndex();
+
+	if (InDataHandle.IsExternalDataType())
+	{
+		// Todo: unsupported external sources
+		return nullptr;
+	}
+
+	const ESceneStateDataType DataType = static_cast<ESceneStateDataType>(InDataHandle.GetDataType());
+
+	switch (DataType)
+	{
+	case ESceneStateDataType::Root:
+		return this;
+
+	case ESceneStateDataType::Task:
+		return TaskInstances[DataIndex].GetScriptStruct();
+
+	case ESceneStateDataType::EventHandler:
+		return EventHandlers[DataIndex].GetEventSchemaHandle().GetEventStruct();
+
+	case ESceneStateDataType::TaskExtension:
+		if (const FSceneStateTask* Task = Tasks[DataIndex].GetPtr<const FSceneStateTask>())
+		{
+			if (const FSceneStateTaskBindingExtension* BindingExtension = Task->GetBindingExtension())
+			{
+				FStructView DataView;
+				if (BindingExtension->FindDataByIndex(TaskInstances[DataIndex], DataSubIndex, DataView))
+				{
+					return DataView.GetScriptStruct();
+				}
+			}
+		}
+		break;
+
+	case ESceneStateDataType::Transition:
+		return TransitionParameters[DataIndex].GetPropertyBagStruct();
+
+	case ESceneStateDataType::StateMachine:
+		return StateMachines[DataIndex].GetParametersStruct();
+	}
+
+	return nullptr;
+}
+
+void USceneStateGeneratedClass::ResolveBindings()
+{
+	if (!IsFullClass())
+	{
+		return;	
+	}
+
+	UE::SceneState::PatchBindingCollection({ .BindingCollection = BindingCollection, .FindDataStructFunctor =
+		[this](const FSceneStateBindingDataHandle& InDataHandle)->const UStruct*
+		{
+			return FindDataStruct(InDataHandle);
+		}});
+
+	// Resolves property paths used by bindings a store property pointers
+	if (!BindingCollection.ResolvePaths())
+	{
+		UE_LOG(LogSceneState, Warning, TEXT("%s failed to resolve bindings. Try compiling Scene State Blueprint again"), *GetName());
+	}
+}
+
+void USceneStateGeneratedClass::Reset()
+{
+	Tasks.Reset();
+	TaskPrerequisites.Reset();
+	TaskInstances.Reset();
+	States.Reset();
+	Conduits.Reset();
+	ConduitLinks.Reset();
+	EventHandlers.Reset();
+	Transitions.Reset();
+	TransitionLinks.Reset();
+	TransitionParameters.Reset();
+	StateMachines.Reset();
+	BindingCollection.Reset();
+	StateMachineIdToIndex.Reset();
+
+#if WITH_EDITOR
+	TaskMetadata.Reset();
+	StateMetadata.Reset();
+	TransitionMetadata.Reset();
+	StateNodeToIndex.Reset();
+	StateMachineGraphToIndex.Reset();
+	TaskNodeToIndex.Reset();
+#endif
+}
+
+bool USceneStateGeneratedClass::IsFullClass() const
+{
+#if WITH_EDITOR
+	FNameBuilder ClassName(GetFName());
+	FStringView ClassNameView = ClassName.ToView();
+	return !ClassNameView.StartsWith(TEXT("SKEL_")) && !ClassNameView.StartsWith(TEXT("REINST_"));
+#else
+	return true;
+#endif
+}
+
+#if WITH_EDITOR
+void USceneStateGeneratedClass::OnObjectsReinstanced(const TMap<UObject*, UObject*>& InReplacementMap)
+{
+	TSet<const UStruct*> Structs;
+	Structs.Reserve(InReplacementMap.Num());
+
+	bool bRequiresResolve = false;
+
+	for (const TPair<UObject*, UObject*>& Pair : InReplacementMap)
+	{
+		if (const UObject* Replacement = Pair.Value)
+		{
+			if (Replacement->IsIn(this))
+			{
+				bRequiresResolve = true;
+				break;
+			}
+			Structs.Add(Replacement->GetClass());
+		}
+	}
+
+	if (bRequiresResolve || BindingCollection.ContainsAnyStruct(Structs))
+	{
+		ResolveBindings();
+	}
+}
+
+void USceneStateGeneratedClass::OnStructsReinstanced(const UUserDefinedStruct& InStruct)
+{
+	TSet<const UStruct*> Structs;
+	Structs.Add(&InStruct);
+
+	if (BindingCollection.ContainsAnyStruct(Structs))
+	{
+		ResolveBindings();
+	}
+}
+#endif
diff --git a/Engine/Plugins/Experimental/SceneState/Source/SceneState/Private/Tasks/SceneStateTask.cpp b/Engine/Plugins/Experimental/SceneState/Source/SceneState/Private/Tasks/SceneStateTask.cpp
new file mode 100644
--- /dev/null
+++ b/Engine/Plugins/Experimental/SceneState/Source/SceneState/Private/Tasks/SceneStateTask.cpp
@@ -0,0 +1,192 @@
+// Copyright Epic Games, Inc. All Rights Reserved.
+
+#include "Tasks/SceneStateTask.h"
+#include "SceneStateBindingUtils.h"
+#include "SceneStateExecutionContext.h"
+#include "SceneStateObject.h"
+#include "SceneStateUtils.h"
+#include "Tasks/SceneStateTaskBindingExtension.h"
+#include "Tasks/SceneStateTaskInstance.h"
+
+namespace UE::SceneState::Private
+{
+
+EExecutionStatus GetTaskStatus(ESceneStateTaskStopReason InStopReason)
+{
+	switch (InStopReason)
+	{
+	case ESceneStateTaskStopReason::State:
+		return EExecutionStatus::NotStarted;
+
+	case ESceneStateTaskStopReason::Finished:
+		return EExecutionStatus::Finished;
+	}
+
+	checkNoEntry();
+	return EExecutionStatus::Finished;
+}
+
+} // UE::SceneState::Private
+
+#if WITH_EDITOR
+const UScriptStruct* FSceneStateTask::GetTaskInstanceType() const
+{
+	return OnGetTaskInstanceType();
+}
+
+void FSceneStateTask::BuildTaskInstance(UObject* InOuter, FStructView InTaskInstance) const
+{
+	if (InTaskInstance.IsValid())
+	{
+		OnBuildTaskInstance(InOuter, InTaskInstance);
+	}
+}
+#endif
+
+const FSceneStateTaskBindingExtension* FSceneStateTask::GetBindingExtension() const
+{
+	if (EnumHasAllFlags(TaskFlags, ESceneStateTaskFlags::HasBindingExtension))
+	{
+		return OnGetBindingExtension();
+	}
+	return nullptr;
+}
+
+FStructView FSceneStateTask::FindTaskInstance(const FSceneStateExecutionContext& InContext) const
+{
+	return InContext.FindTaskInstance(GetTaskIndex());
+}
+
+void FSceneStateTask::Setup(const FSceneStateExecutionContext& InContext, FStructView InTaskInstance) const
+{
+	using namespace UE::SceneState;
+
+	QUICK_SCOPE_CYCLE_COUNTER(SceneState_OnSetup);
+
+	FSceneStateTaskInstance* Instance = InTaskInstance.GetPtr<FSceneStateTaskInstance>();
+	if (!Instance)
+	{
+		return;
+	}
+
+	Instance->Status = EExecutionStatus::NotStarted;
+	OnSetup(InContext, InTaskInstance);
+}
+
+void FSceneStateTask::Start(const FSceneStateExecutionContext& InContext, FStructView InTaskInstance) const
+{
+	using namespace UE::SceneState;
+
+	QUICK_SCOPE_CYCLE_COUNTER(SceneState_OnStart);
+
+	FSceneStateTaskInstance* Instance = InTaskInstance.GetPtr<FSceneStateTaskInstance>();
+	if (!Instance || Instance->Status != EExecutionStatus::NotStarted)
+	{
+		return;
+	}
+
+	Instance->Status = EExecutionStatus::Running;
+	Instance->Result = ESceneStateTaskResult::Undetermined;
+	Instance->ElapsedTime = 0.f;
+
+	ApplyBindings(InContext, InTaskInstance);
+
+	OnStart(InContext, InTaskInstance);
+}
+
+void FSceneStateTask::Tick(const FSceneStateExecutionContext& InContext, FStructView InTaskInstance, float InDeltaSeconds) const
+{
+	using namespace UE::SceneState;
+
+	QUICK_SCOPE_CYCLE_COUNTER(SceneState_OnTick);
+
+	FSceneStateTaskInstance* Instance = InTaskInstance.GetPtr<FSceneStateTaskInstance>();
+	if (!Instance || Instance->Status != EExecutionStatus::Running)
+	{
+		return;
+	}
+
+	// Even if Task doesn't tick the virtual/bp functions, keep elapsed time tracked
+	Instance->ElapsedTime += InDeltaSeconds;
+
+	if (EnumHasAllFlags(TaskFlags, ESceneStateTaskFlags::Ticks))
+	{
+		OnTick(InContext, InTaskInstance, InDeltaSeconds);
+	}
+}
+
+void FSceneStateTask::Stop(const FSceneStateExecutionContext& InContext, FStructView InTaskInstance, ESceneStateTaskStopReason InStopReason) const
+{
+	using namespace UE::SceneState;
+
+	QUICK_SCOPE_CYCLE_COUNTER(SceneState_OnStop);
+
+	FSceneStateTaskInstance* Instance = InTaskInstance.GetPtr<FSceneStateTaskInstance>();
+	if (!Instance)
+	{
+		return;
+	}
+
+	if (Instance->Status == EExecutionStatus::Running)
+	{
+		OnStop(InContext, InTaskInstance, InStopReason);
+	}
+
+	Instance->ElapsedTime = 0.f;
+	Instance->Status = Private::GetTaskStatus(InStopReason);
+}
+
+void FSceneStateTask::Finish(const FSceneStateExecutionContext& InContext, FStructView InTaskInstance) const
+{
+	Stop(InContext, InTaskInstance, ESceneStateTaskStopReason::Finished);
+}
+
+void FSceneStateTask::SetFlags(ESceneStateTaskFlags InFlags)
+{
+	TaskFlags |= InFlags;
+}
+
+void FSceneStateTask::ClearFlags(ESceneStateTaskFlags InFlags)
+{
+	TaskFlags &= ~InFlags;
+}
+
+bool FSceneStateTask::ApplyBindings(const FSceneStateExecutionContext& InContext, FStructView InTaskInstance) const
+{
+	QUICK_SCOPE_CYCLE_COUNTER(SceneStateTask_ApplyBindings);
+
+	const FSceneStateBindingCollection* const BindingCollection = InContext.GetBindingCollection();
+	if (!BindingCollection)
+	{
+		return false;
+	}
+
+	auto FindDataView =
+		[&InContext](const FSceneStateBindingDataHandle& InDataHandle)
+		{
+			return InContext.FindDataView(InDataHandle);
+		};
+
+	UE::SceneState::FApplyBatchParams ApplyBatchParams
+		{
+			.BatchIndex = BindingsBatch.Get(),
+			.TargetDataView = InTaskInstance,
+			.BindingCollection = *BindingCollection,
+			.FindSourceDataViewFunctor = FindDataView,
+		};
+
+	bool bResult = ApplyBatch(ApplyBatchParams);
+
+	if (const FSceneStateTaskBindingExtension* BindingExtension = GetBindingExtension())
+	{
+		BindingExtension->VisitBindingBatches(InTaskInstance,
+			[&bResult, &ApplyBatchParams](uint16 InBatchIndex, FStructView InTargetDataView)
+			{
+				ApplyBatchParams.BatchIndex = InBatchIndex;
+				ApplyBatchParams.TargetDataView = InTargetDataView;
+				bResult &= ApplyBatch(ApplyBatchParams);
+			});
+	}
+
+	return bResult;
+}
diff --git a/Engine/Plugins/Experimental/SceneState/Source/SceneState/Private/Tasks/SceneStateTaskExecutionContext.cpp b/Engine/Plugins/Experimental/SceneState/Source/SceneState/Private/Tasks/SceneStateTaskExecutionContext.cpp
new file mode 100644
--- /dev/null
+++ b/Engine/Plugins/Experimental/SceneState/Source/SceneState/Private/Tasks/SceneStateTaskExecutionContext.cpp
@@ -0,0 +1,146 @@
+// Copyright Epic Games, Inc. All Rights Reserved.
+
+#include "Tasks/SceneStateTaskExecutionContext.h"
+#include "SceneStateExecutionContext.h"
+#include "SceneStateExecutionContextRegistry.h"
+#include "SceneStateLog.h"
+#include "Tasks/SceneStateTask.h"
+#include "Tasks/SceneStateTaskInstance.h"
+
+namespace UE::SceneState
+{
+
+namespace Private
+{
+
+FExecutionContextHandle GetContextHandle(const FSceneStateExecutionContext& InContext)
+{
+	if (TSharedPtr<const FExecutionContextRegistry> ContextRegistry = InContext.GetContextRegistry().Pin())
+	{
+		return ContextRegistry->FindHandle(InContext);
+	}
+	return FExecutionContextHandle();
+}
+
+uint32 GetTaskInstanceId(const FSceneStateTask& InTask, const FSceneStateExecutionContext& InContext)
+{
+	if (FSceneStateTaskInstance* TaskInstance = InContext.FindTaskInstance(InTask.GetTaskIndex()).GetPtr<FSceneStateTaskInstance>())
+	{
+		return TaskInstance->GetInstanceId();
+	}
+	return 0;
+}
+
+bool IsMatchingInstanceId(FConstStructView InTaskInstanceView, uint32 InTaskInstanceId)
+{
+	if (const FSceneStateTaskInstance* TaskInstance = InTaskInstanceView.GetPtr<const FSceneStateTaskInstance>())
+	{
+		return TaskInstance->GetInstanceId() == InTaskInstanceId;
+	}
+	return false;
+}
+
+} // Private
+
+FTaskExecutionContext::FTaskExecutionContext()
+	: TaskIndex(FSceneStateRange::InvalidIndex)
+	, TaskInstanceId(0)
+{
+}
+
+FTaskExecutionContext::FTaskExecutionContext(const FSceneStateTask& InTask, const FSceneStateExecutionContext& InContext)
+	: ContextHandle(Private::GetContextHandle(InContext))
+	, ContextRegistryWeak(InContext.GetContextRegistry())
+	, TaskIndex(InTask.GetTaskIndex())
+	, TaskInstanceId(Private::GetTaskInstanceId(InTask, InContext))
+{
+}
+
+const FSceneStateExecutionContext* FTaskExecutionContext::GetExecutionContext() const
+{
+	if (const TSharedPtr<const FExecutionContextRegistry> ContextRegistry = ContextRegistryWeak.Pin())
+	{
+		return ContextRegistry->FindContext(ContextHandle);
+	}
+	return nullptr;
+}
+
+FConstStructView FTaskExecutionContext::GetTask() const
+{
+	if (const FSceneStateExecutionContext* ExecutionContext = GetExecutionContext())
+	{
+		// Ensure this Context is still valid
+		if (Private::IsMatchingInstanceId(ExecutionContext->FindTaskInstance(TaskIndex), TaskInstanceId))
+		{
+			return ExecutionContext->FindTask(TaskIndex);
+		}
+	}
+	return FConstStructView();
+}
+
+FStructView FTaskExecutionContext::GetTaskInstance() const
+{
+	if (const FSceneStateExecutionContext* ExecutionContext = GetExecutionContext())
+	{
+		FStructView TaskInstance = ExecutionContext->FindTaskInstance(TaskIndex);
+
+		// Ensure this Context is still valid
+		if (Private::IsMatchingInstanceId(TaskInstance, TaskInstanceId))
+		{
+			return TaskInstance;
+		}
+	}
+	return FStructView();
+}
+
+void FTaskExecutionContext::FinishTask() const
+{
+	constexpr const TCHAR* BaseErrorMessage = TEXT("Failed to Finish Task.");
+
+	const FSceneStateExecutionContext* ExecutionContext = GetExecutionContext();
+	if (!ExecutionContext)
+	{
+		UE_LOG(LogSceneState, Error, TEXT("[%s] %s Invalid Execution Context !")
+			, *GetDebugString()
+			, BaseErrorMessage)
+		return;
+	}
+
+	const FSceneStateTask* Task = ExecutionContext->FindTask(TaskIndex).GetPtr<const FSceneStateTask>();
+	if (!Task)
+	{
+		UE_LOG(LogSceneState, Error, TEXT("[%s in Context '%s'] %s Task not found!")
+			, *GetDebugString()
+			, BaseErrorMessage
+			, *ExecutionContext->GetExecutionContextName())
+		return;
+	}
+
+	const FStructView TaskInstance = ExecutionContext->FindTaskInstance(TaskIndex);
+	if (!TaskInstance.IsValid())
+	{
+		UE_LOG(LogSceneState, Error, TEXT("[%s in Context '%s'] %s Task Instance not found!")
+			, *GetDebugString()
+			, *ExecutionContext->GetExecutionContextName()
+			, BaseErrorMessage)
+		return;
+	}
+
+	if (!Private::IsMatchingInstanceId(TaskInstance, TaskInstanceId))
+	{
+		UE_LOG(LogSceneState, Error, TEXT("[%s in Context '%s'] %s Task Context outdated!")
+			, *GetDebugString()
+			, *ExecutionContext->GetExecutionContextName()
+			, BaseErrorMessage)
+		return;
+	}
+
+	Task->Finish(*ExecutionContext, TaskInstance);
+}
+
+FString FTaskExecutionContext::GetDebugString() const
+{
+	return FString::Printf(TEXT("Task Index: %d, State Instance Id: %u"), TaskIndex, TaskInstanceId);
+}
+
+} // UE::SceneState
diff --git a/Engine/Plugins/Experimental/SceneState/Source/SceneState/Private/Tasks/SceneStateTaskInstance.cpp b/Engine/Plugins/Experimental/SceneState/Source/SceneState/Private/Tasks/SceneStateTaskInstance.cpp
new file mode 100644
--- /dev/null
+++ b/Engine/Plugins/Experimental/SceneState/Source/SceneState/Private/Tasks/SceneStateTaskInstance.cpp
@@ -0,0 +1,19 @@
+// Copyright Epic Games, Inc. All Rights Reserved.
+
+#include "Tasks/SceneStateTaskInstance.h"
+
+namespace UE::SceneState::Private
+{
+	
+uint16 GetNextInstanceId()
+{
+	static uint16 Id = 0;
+	return ++Id; // Starts with 1
+}
+
+} // UE::SceneState::Private
+
+FSceneStateTaskInstance::FSceneStateTaskInstance()
+	: InstanceId(UE::SceneState::Private::GetNextInstanceId())
+{
+}
diff --git a/Engine/Plugins/Experimental/SceneState/Source/SceneState/Public/SceneStateGeneratedClass.h b/Engine/Plugins/Experimental/SceneState/Source/SceneState/Public/SceneStateGeneratedClass.h
new file mode 100644
--- /dev/null
+++ b/Engine/Plugins/Experimental/SceneState/Source/SceneState/Public/SceneStateGeneratedClass.h
@@ -0,0 +1,241 @@
+// Copyright Epic Games, Inc. All Rights Reserved.
+
+#pragma once
+
+#include "Conduit/SceneStateConduit.h"
+#include "Conduit/SceneStateConduitLink.h"
+#include "Engine/BlueprintGeneratedClass.h"
+#include "SceneState.h"
+#include "SceneStateBindingCollection.h"
+#include "SceneStateEventHandler.h"
+#include "SceneStateMachine.h"
+#include "SceneStateRange.h"
+#include "StructUtils/InstancedStructContainer.h"
+#include "Tasks/SceneStateTask.h"
+#include "Tasks/SceneStateTaskMetadata.h"
+#include "Transition/SceneStateTransition.h"
+#include "Transition/SceneStateTransitionLink.h"
+#include "Transition/SceneStateTransitionMetadata.h"
+#include "UObject/ObjectKey.h"
+#include "UObject/Package.h"
+#include "SceneStateGeneratedClass.generated.h"
+
+class USceneStateObject;
+class UUserDefinedStruct;
+
+namespace UE::SceneState::Editor
+{
+	class FBindingCompiler;
+	class FBlueprintCompilerContext;
+	class FStateMachineCompiler;
+}
+
+/**
+ * Object Class for the Scene State Object
+ * Holds all the data about States, State Machines, Tasks, etc.
+ * All this data is immutable in execution, and as such, it is not instanced to the Scene State Object instances.
+ * @see FSceneStateExecutionContext
+ */
+UCLASS(MinimalAPI)
+class USceneStateGeneratedClass : public UBlueprintGeneratedClass
+{
+	GENERATED_BODY()
+
+public:
+	SCENESTATE_API USceneStateGeneratedClass();
+
+	const FSceneState* GetRootState() const;
+
+	/** Finds the Top-Level State Machine that is linked to the given Guid */
+	const FSceneStateMachine* FindStateMachine(const FGuid& InStateMachineId) const;
+
+	const FSceneStateBindingCollection& GetBindingCollection() const
+	{
+		return BindingCollection;
+	}
+
+	TConstArrayView<FSceneState> GetStates() const
+	{
+		return States;
+	}
+
+	TConstArrayView<FSceneStateConduit> GetConduits() const
+	{
+		return Conduits;
+	}
+
+	const FInstancedStructContainer& GetTasks() const
+	{
+		return Tasks;
+	}
+
+	TConstArrayView<uint16> GetTaskPrerequisites() const
+	{
+		return TaskPrerequisites;
+	}
+
+	const FInstancedStructContainer& GetTaskInstances() const
+	{
+		return TaskInstances;
+	}
+
+	TConstArrayView<FSceneStateEventHandler> GetEventHandlers() const
+	{
+		return EventHandlers;
+	}
+
+	TConstArrayView<FSceneStateTransition> GetTransitions() const
+	{
+		return Transitions;
+	}
+
+	FInstancedPropertyBag GetTransitionParameter(uint16 InAbsoluteIndex) const;
+
+	TConstArrayView<FSceneStateMachine> GetStateMachines() const
+	{
+		return StateMachines;
+	}
+
+#if WITH_EDITOR
+	TConstArrayView<FSceneStateMetadata> GetStateMetadata() const
+	{
+		return StateMetadata;
+	}
+
+	TConstArrayView<FSceneStateTaskMetadata> GetTaskMetadata() const
+	{
+		return TaskMetadata;
+	}
+#endif
+
+	//~ Begin UClass
+	SCENESTATE_API virtual void Link(FArchive& Ar, bool bInRelinkExistingProperties) override;
+	//~ End UClass
+
+	//~ Begin UObject
+	SCENESTATE_API virtual void PostLoad() override;
+	SCENESTATE_API virtual void BeginDestroy() override;
+	//~ End UObject
+
+#if WITH_EDITOR
+	/** Finds the mapped compiled state for a given state node */
+	SCENESTATE_API const FSceneState* FindStateFromNode(FObjectKey InStateNode) const;
+
+	/** Finds the mapped compiled task for a given task node */
+	SCENESTATE_API const FSceneStateTask* FindTaskFromNode(FObjectKey InTaskNode) const;
+
+	/**
+	 * For a given root state and state node retrieves the mapped state instances
+	 * @param InRootState the root state that is outered to a player
+	 * @param InStateNode the state node to look for 
+	 * @param InFunctor the functor to run for each state instance
+	 */
+	SCENESTATE_API void ForEachStateInstance(USceneStateObject& InRootState
+		, FObjectKey InStateNode
+		, TFunctionRef<void(const FSceneStateInstance&)> InFunctor) const;
+
+	/**
+	 * For a given root state and task node retrieves the mapped task instances
+	 * @param InRootState the root state that is outered to a player
+	 * @param InTaskNode the task node to look for
+	 * @param InFunctor the functor to run for each task instance
+	 */
+	SCENESTATE_API void ForEachTaskInstance(USceneStateObject& InRootState
+		, FObjectKey InTaskNode
+		, TFunctionRef<void(const FSceneStateTaskInstance&)> InFunctor) const;
+#endif
+
+private:
+	/** Finds the struct type for the given data handle */
+	const UStruct* FindDataStruct(const FSceneStateBindingDataHandle& InDataHandle);
+
+	/** Patches bindings and resolves the binding paths for the owning binding collection */
+	SCENESTATE_API void ResolveBindings();
+
+	/** Resets all the elements that get compiled for this generated class (e.g. Tasks, State Machines, etc) */
+	SCENESTATE_API void Reset();
+
+	bool IsFullClass() const;
+
+#if WITH_EDITOR
+	void OnObjectsReinstanced(const TMap<UObject*, UObject*>& InReplacementMap);
+	void OnStructsReinstanced(const UUserDefinedStruct& InStruct);
+#endif
+
+	UPROPERTY()
+	FSceneStateBindingCollection BindingCollection;
+
+	UPROPERTY()
+	uint16 RootStateIndex = FSceneStateRange::InvalidIndex;
+
+	UPROPERTY()
+	TArray<FSceneState> States;
+
+	UPROPERTY()
+	FInstancedStructContainer Tasks;
+
+	/** Array of the task prerequisites in their relative index. Each task has a unique range indicating their pre-requisites */
+	UPROPERTY()
+	TArray<uint16> TaskPrerequisites;
+
+	/** Templates used to instantiate the Task Instances */
+	UPROPERTY()
+	FInstancedStructContainer TaskInstances;
+
+	UPROPERTY()
+	TArray<FSceneStateEventHandler> EventHandlers;
+
+	UPROPERTY()
+	TArray<FSceneStateTransition> Transitions;
+
+	/** Compiled transition information only used in Link time */
+	UPROPERTY()
+	TArray<FSceneStateTransitionLink> TransitionLinks;
+
+	/** Map of the Transition Index (absolute) to the template transition parameters map for evaluation function call. */
+	UPROPERTY()
+	TMap<uint16, FInstancedPropertyBag> TransitionParameters;
+
+	/** All the compiled conduits */
+	UPROPERTY()
+	TArray<FSceneStateConduit> Conduits;
+
+	/** Compiled conduit information used only in Link time */
+	UPROPERTY()
+	TArray<FSceneStateConduitLink> ConduitLinks;
+
+	UPROPERTY()
+	TArray<FSceneStateMachine> StateMachines;
+
+	/** Map of the Top-Level State machine Parameters id to the index in the state machine array */
+	UPROPERTY()
+	TMap<FGuid, uint16> StateMachineIdToIndex;
+
+#if WITH_EDITORONLY_DATA
+	UPROPERTY()
+	TArray<FSceneStateMetadata> StateMetadata;
+
+	UPROPERTY()
+	TArray<FSceneStateTaskMetadata> TaskMetadata;
+
+	UPROPERTY()
+	TArray<FSceneStateTransitionMetadata> TransitionMetadata;
+
+	/** Map of a State Node to its State Index */
+	TMap<FObjectKey, uint16> StateNodeToIndex;
+
+	/** Map of a State Machine Graph to its State Machine Index */
+	TMap<FObjectKey, uint16> StateMachineGraphToIndex;
+
+	/** Map of a Task Node to its Task Index */
+	TMap<FObjectKey, uint16> TaskNodeToIndex;
+
+	FDelegateHandle OnObjectsReinstancedHandle;
+	FDelegateHandle OnStructsReinstancedHandle;
+#endif
+
+	friend UE::SceneState::Editor::FBindingCompiler;
+	friend UE::SceneState::Editor::FBlueprintCompilerContext;
+	friend UE::SceneState::Editor::FStateMachineCompiler;
+	friend class USceneStateBlueprint;
+};
diff --git a/Engine/Plugins/Experimental/SceneState/Source/SceneState/Public/Tasks/SceneStateTaskInstance.h b/Engine/Plugins/Experimental/SceneState/Source/SceneState/Public/Tasks/SceneStateTaskInstance.h
new file mode 100644
--- /dev/null
+++ b/Engine/Plugins/Experimental/SceneState/Source/SceneState/Public/Tasks/SceneStateTaskInstance.h
@@ -0,0 +1,35 @@
+// Copyright Epic Games, Inc. All Rights Reserved.
+
+#pragma once
+
+#include "SceneStateEnums.h"
+#include "SceneStateTaskInstance.generated.h"
+
+struct FPropertyBindingDataView;
+
+/** Represents the Instance Data of a Task */
+USTRUCT()
+struct FSceneStateTaskInstance
+{
+	GENERATED_BODY()
+
+	SCENESTATE_API FSceneStateTaskInstance();
+
+	/** Elapsed Time */
+	float ElapsedTime = 0.f;
+
+	/** Current status of the Task */
+	UE::SceneState::EExecutionStatus Status = UE::SceneState::EExecutionStatus::NotStarted;
+
+	/** Result of the Task, when finished, otherwise undetermined */
+	ESceneStateTaskResult Result = ESceneStateTaskResult::Undetermined;
+
+	uint16 GetInstanceId() const
+	{
+		return InstanceId;
+	}
+
+private:
+	/** The id for this instance. This is used to differentiate instances for the same task */
+	uint16 InstanceId;
+};
diff --git a/Engine/Plugins/Experimental/SceneState/Source/SceneStateBinding/Public/SceneStateBindingCollection.h b/Engine/Plugins/Experimental/SceneState/Source/SceneStateBinding/Public/SceneStateBindingCollection.h
new file mode 100644
--- /dev/null
+++ b/Engine/Plugins/Experimental/SceneState/Source/SceneStateBinding/Public/SceneStateBindingCollection.h
@@ -0,0 +1,60 @@
+// Copyright Epic Games, Inc. All Rights Reserved.
+
+#pragma once
+
+#include "PropertyBindingBindingCollection.h"
+#include "SceneStateBinding.h"
+#include "SceneStateBindingDesc.h"
+#include "SceneStateBindingCollection.generated.h"
+
+namespace UE::SceneState::Editor
+{
+	class FBindingCompiler;
+}
+
+USTRUCT()
+struct FSceneStateBindingCollection : public FPropertyBindingBindingCollection
+{
+	GENERATED_BODY()
+
+	TArrayView<FSceneStateBindingDesc> GetMutableBindingDescs()
+	{
+		return BindingDescs;
+	}
+
+	TArrayView<FSceneStateBinding> GetMutableBindings()
+	{
+		return Bindings;
+	}
+
+	/** Finds the binding desc matching the given data handle */
+	SCENESTATEBINDING_API const FSceneStateBindingDesc* FindBindingDesc(FSceneStateBindingDataHandle InDataHandle) const;
+
+#if WITH_EDITOR
+	//~ Begin FPropertyBindingBindingCollection
+	SCENESTATEBINDING_API virtual FPropertyBindingBinding* AddBindingInternal(const FPropertyBindingPath& InSourcePath, const FPropertyBindingPath& InTargetPath) override;
+	SCENESTATEBINDING_API virtual void RemoveBindingsInternal(TFunctionRef<bool(FPropertyBindingBinding&)> InPredicate) override;
+	SCENESTATEBINDING_API virtual bool HasBindingInternal(TFunctionRef<bool(const FPropertyBindingBinding&)> InPredicate) const override;
+	SCENESTATEBINDING_API virtual const FPropertyBindingBinding* FindBindingInternal(TFunctionRef<bool(const FPropertyBindingBinding&)> InPredicate) const override;
+#endif
+	SCENESTATEBINDING_API virtual int32 GetNumBindings() const override;
+	SCENESTATEBINDING_API virtual int32 GetNumBindableStructDescriptors() const override;
+	SCENESTATEBINDING_API virtual const FPropertyBindingBindableStructDescriptor* GetBindableStructDescriptorFromHandle(FConstStructView InSourceHandleView) const override;
+	SCENESTATEBINDING_API virtual void ForEachBinding(TFunctionRef<void(const FPropertyBindingBinding& Binding)> InFunction) const override;
+	SCENESTATEBINDING_API virtual void ForEachBinding(FPropertyBindingIndex16 InBegin, FPropertyBindingIndex16 InEnd, TFunctionRef<void(const FPropertyBindingBinding& Binding, const int32 BindingIndex)> InFunction) const override;
+	SCENESTATEBINDING_API virtual void ForEachMutableBinding(TFunctionRef<void(FPropertyBindingBinding& Binding)> InFunction) override;
+	SCENESTATEBINDING_API virtual void VisitBindings(TFunctionRef<EVisitResult(const FPropertyBindingBinding& Binding)> InFunction) const override;
+	SCENESTATEBINDING_API virtual void VisitMutableBindings(TFunctionRef<EVisitResult(FPropertyBindingBinding& Binding)> InFunction) override;
+	SCENESTATEBINDING_API virtual void OnReset() override;
+	SCENESTATEBINDING_API virtual void VisitSourceStructDescriptorInternal(TFunctionRef<EVisitResult(const FPropertyBindingBindableStructDescriptor&)> InFunction) const override;
+	//~ End FPropertyBindingBindingCollection
+
+private:
+	UPROPERTY()
+	TArray<FSceneStateBindingDesc> BindingDescs;
+
+	UPROPERTY()
+	TArray<FSceneStateBinding> Bindings;
+
+	friend UE::SceneState::Editor::FBindingCompiler;
+};
diff --git a/Engine/Plugins/Experimental/SceneState/Source/SceneStateBinding/Public/SceneStateBindingUtils.h b/Engine/Plugins/Experimental/SceneState/Source/SceneStateBinding/Public/SceneStateBindingUtils.h
new file mode 100644
--- /dev/null
+++ b/Engine/Plugins/Experimental/SceneState/Source/SceneStateBinding/Public/SceneStateBindingUtils.h
@@ -0,0 +1,60 @@
+// Copyright Epic Games, Inc. All Rights Reserved.
+
+#pragma once
+
+#include "PropertyBindingDataView.h"
+#include "Templates/Function.h"
+
+class UObject;
+struct FGuid;
+struct FSceneStateBindingCollection;
+struct FSceneStateBindingDataHandle;
+
+namespace UE::SceneState
+{
+
+#if WITH_EDITOR
+	/**
+	 * Update the bindings of a source struct id into a target one.
+	 * It goes through the given object's outer chain to get to the binding collection owner to do this operation
+	 * @param InObject object to use to find the binding collection owner
+	 * @param InOldStructId the struct id to queue copy bindings from
+	 * @param InNewStructId the struct id to queue copy bindings to
+	 */
+	SCENESTATEBINDING_API void HandleStructIdChanged(UObject& InObject, const FGuid& InOldStructId, const FGuid& InNewStructId);
+#endif
+
+	struct FApplyBatchParams
+	{
+		/** Batch within the Binding Collection to run */
+		uint16 BatchIndex;
+		/** Target Data to copy the source data to */
+		FPropertyBindingDataView TargetDataView;
+		/** Binding collection holding the Copy Batch to run */
+		const FSceneStateBindingCollection& BindingCollection;
+		/** Functor to find the Source Data View for a given Data Handle */
+		TFunctionRef<FPropertyBindingDataView(const FSceneStateBindingDataHandle&)> FindSourceDataViewFunctor;
+	};
+	/** Runs the Batch within the Binding Collection for a given Target Data */
+	SCENESTATEBINDING_API bool ApplyBatch(const FApplyBatchParams& InParams);
+
+	struct FPatchBindingParams
+	{
+		/** Binding collection to patch */
+		FSceneStateBindingCollection& BindingCollection;
+		/** Functor to find the data struct for a given Data Handle */
+		TFunctionRef<const UStruct*(const FSceneStateBindingDataHandle&)> FindDataStructFunctor;
+	};
+	/** Patches invalidated structs (property bags, user defined structs, etc.) in the given binding collection */
+	SCENESTATEBINDING_API void PatchBindingCollection(const FPatchBindingParams& InParams);
+
+	/** Patches any invalidated structs (property bags, user defined structs, etc.) in the binding descs of the given binding collection */
+	void PatchBindingDescs(const FPatchBindingParams& InParams);
+
+	/** Patches any invalidated structs (property bags, user defined structs, etc.) in the bindings of the given binding collection */
+	void PatchBindings(const FPatchBindingParams& InParams);
+
+	/** Patches any invalidated structs (property bags, user defined structs, etc.) in the copy batches of the given binding collection */
+	void PatchCopyBatches(const FPatchBindingParams& InParams);
+
+} // UE::SceneState
diff --git a/Engine/Plugins/Experimental/SceneState/Source/SceneStateBlueprint/Public/Actions/SceneStateBlueprintAction_Graph.h b/Engine/Plugins/Experimental/SceneState/Source/SceneStateBlueprint/Public/Actions/SceneStateBlueprintAction_Graph.h
new file mode 100644
--- /dev/null
+++ b/Engine/Plugins/Experimental/SceneState/Source/SceneStateBlueprint/Public/Actions/SceneStateBlueprintAction_Graph.h
@@ -0,0 +1,48 @@
+// Copyright Epic Games, Inc. All Rights Reserved.
+
+#pragma once
+
+#include "EdGraphSchema_K2_Actions.h"
+#include "Textures/SlateIcon.h"
+
+namespace UE::SceneState::Graph
+{
+
+struct FBlueprintAction_Graph : public FEdGraphSchemaAction_K2Graph
+{
+	FBlueprintAction_Graph() = default;
+
+	explicit FBlueprintAction_Graph(EEdGraphSchemaAction_K2Graph::Type InType
+		, const FText& InNodeCategory
+		, const FText& InMenuDesc
+		, const FText& InToolTip
+		, const FSlateIcon& InIcon
+		, const int32 InGrouping
+		, const int32 InSectionID = 0)
+		: FEdGraphSchemaAction_K2Graph(InType, InNodeCategory, InMenuDesc, InToolTip, InGrouping, InSectionID)
+		, Icon(InIcon)
+	{
+	}
+
+	static FName StaticGetTypeId()
+	{
+		static FName Type = TEXT("FEdGraphSchemaAction");
+		return Type;
+	}
+
+	//~ Begin FEdGraphSchemaAction
+	virtual const FSlateBrush* GetPaletteIcon() const override
+	{
+		return Icon.GetIcon();
+	}
+
+	virtual FName GetTypeId() const override
+	{
+		return StaticGetTypeId();
+	}
+	//~ End FEdGraphSchemaAction
+
+	const FSlateIcon Icon;
+};
+
+} // UE::SceneState::Editor
diff --git a/Engine/Plugins/Experimental/SceneState/Source/SceneStateBlueprint/Public/SceneStateBlueprint.h b/Engine/Plugins/Experimental/SceneState/Source/SceneStateBlueprint/Public/SceneStateBlueprint.h
new file mode 100644
--- /dev/null
+++ b/Engine/Plugins/Experimental/SceneState/Source/SceneStateBlueprint/Public/SceneStateBlueprint.h
@@ -0,0 +1,80 @@
+// Copyright Epic Games, Inc. All Rights Reserved.
+
+#pragma once
+
+#include "Containers/Array.h"
+#include "Engine/Blueprint.h"
+#include "PropertyBindingBindingCollectionOwner.h"
+#include "SceneStateBindingCollection.h"
+#include "SceneStateBlueprint.generated.h"
+
+class USceneStateGeneratedClass;
+class USceneStateMachineStateNode;
+
+UCLASS(MinimalAPI)
+class USceneStateBlueprint : public UBlueprint, public IPropertyBindingBindingCollectionOwner
+{
+	GENERATED_BODY()
+
+public:
+	SCENESTATEBLUEPRINT_API USceneStateBlueprint(const FObjectInitializer& InObjectInitializer);
+
+	const FGuid& GetRootId() const
+	{
+		return RootId;
+	}
+
+	template<typename T>
+	T* FindExtension() const
+	{
+		return Cast<T>(FindExtension(T::StaticClass()));
+	}
+
+	SCENESTATEBLUEPRINT_API UBlueprintExtension* FindExtension(TSubclassOf<UBlueprintExtension> InClass) const;
+
+	SCENESTATEBLUEPRINT_API FSceneStateBindingDesc CreateRootBinding() const;
+
+	//~ Begin IPropertyBindingBindingCollectionOwner
+	SCENESTATEBLUEPRINT_API virtual void GetBindableStructs(const FGuid InTargetStructId, TArray<TInstancedStruct<FPropertyBindingBindableStructDescriptor>>& OutStructDescs) const override;
+	SCENESTATEBLUEPRINT_API virtual bool GetBindableStructByID(const FGuid InStructId, TInstancedStruct<FPropertyBindingBindableStructDescriptor>& OutStructDesc) const override;
+	SCENESTATEBLUEPRINT_API virtual bool GetBindingDataViewByID(const FGuid InStructId, FPropertyBindingDataView& OutDataView) const override;
+	SCENESTATEBLUEPRINT_API virtual FPropertyBindingBindingCollection* GetEditorPropertyBindings() override;
+	SCENESTATEBLUEPRINT_API virtual const FPropertyBindingBindingCollection* GetEditorPropertyBindings() const override;
+	SCENESTATEBLUEPRINT_API virtual bool CanCreateParameter(const FGuid InStructId) const override;
+	SCENESTATEBLUEPRINT_API virtual void CreateParametersForStruct(const FGuid InStructId, TArrayView<UE::PropertyBinding::FPropertyCreationDescriptor> InOutCreationDescs) override;
+	//~ End IPropertyBindingBindingCollectionOwner
+
+	//~ Begin UBlueprint
+	SCENESTATEBLUEPRINT_API virtual void SetObjectBeingDebugged(UObject* InNewObject) override;
+	SCENESTATEBLUEPRINT_API virtual UClass* GetBlueprintClass() const override;
+	SCENESTATEBLUEPRINT_API virtual void GetReparentingRules(TSet<const UClass*>& OutAllowedChildrenOfClasses, TSet<const UClass*>& OutDisallowedChildrenOfClasses) const override;
+	SCENESTATEBLUEPRINT_API virtual bool SupportedByDefaultBlueprintFactory() const override;
+	SCENESTATEBLUEPRINT_API virtual void LoadModulesRequiredForCompilation() override;
+	SCENESTATEBLUEPRINT_API virtual bool IsValidForBytecodeOnlyRecompile() const override;
+	//~ End UBlueprint
+
+	//~ Begin UObject
+	SCENESTATEBLUEPRINT_API virtual void BeginDestroy() override;
+	//~ End UObject
+
+	/** The top level State Machine Graphs of this State. Does not include the state machine graphs of nested state nodes */
+	UPROPERTY()
+	TArray<TObjectPtr<UEdGraph>> StateMachineGraphs;
+
+	/** Holds all the editor bindings prior to compilation */
+	UPROPERTY()
+	FSceneStateBindingCollection BindingCollection;
+
+private:
+	/** Called when a blueprint variable has been renamed. Used to fix the bindings to the variable being renamed to use the new name */
+	void OnRenameVariableReferences(UBlueprint* InBlueprint, UClass* InVariableClass, const FName& InOldVariableName, const FName& InNewVariableName);
+
+	/** Called when a blueprint variable has been renamed. Used to fix the path to the variable being renamed to use the new name */
+	void RenameVariableReferenceInPath(FPropertyBindingPath& InPath, FName InOldVariableName, FName InNewVariableName);
+
+	/** Unique id representing this blueprint / generated class as a bindable struct */
+	UPROPERTY()
+	FGuid RootId;
+
+	FDelegateHandle OnRenameVariableReferencesHandle;
+};
diff --git a/Engine/Plugins/Experimental/SceneState/Source/SceneStateBlueprintEditor/Private/Compiler/SceneStateBindingCompiler.cpp b/Engine/Plugins/Experimental/SceneState/Source/SceneStateBlueprintEditor/Private/Compiler/SceneStateBindingCompiler.cpp
new file mode 100644
--- /dev/null
+++ b/Engine/Plugins/Experimental/SceneState/Source/SceneStateBlueprintEditor/Private/Compiler/SceneStateBindingCompiler.cpp
@@ -0,0 +1,387 @@
+// Copyright Epic Games, Inc. All Rights Reserved.
+
+#include "SceneStateBindingCompiler.h"
+#include "SceneStateBlueprint.h"
+#include "SceneStateBlueprintCompilerContext.h"
+#include "SceneStateEventSchema.h"
+#include "SceneStateGeneratedClass.h"
+#include "StructUtils/UserDefinedStruct.h"
+#include "Tasks/SceneStateBlueprintableTask.h"
+#include "Tasks/SceneStateTaskBindingExtension.h"
+
+#define LOCTEXT_NAMESPACE "SceneStateBindingCompiler"
+
+namespace UE::SceneState::Editor
+{
+
+void FBindingCompiler::Compile()
+{
+	check(Blueprint && Class);
+
+	// Initial estimate count of the binding descs
+	const int32 BindingDescInitialCount = 1 + Class->Tasks.Num() + Class->EventHandlers.Num() + Class->TransitionParameters.Num();
+
+	// Copy Editor Bindings and reserve for the binding descs to add
+	Class->BindingCollection.Bindings = Blueprint->BindingCollection.Bindings;
+	Class->BindingCollection.BindingDescs.Empty(BindingDescInitialCount);
+
+	ValidBindingMap.Reserve(BindingDescInitialCount);
+
+	AddRootBindingDesc();
+	AddStateMachineBindingDescs();
+	AddTransitionBindingDescs();
+	AddTaskBindingDescs();
+	AddEventHandlerBindingDescs();
+
+	RemoveInvalidBindings();
+
+	ResolveBindingDataHandles();
+	BatchCopies();
+}
+
+void FBindingCompiler::AddBindingDesc(FSceneStateBindingDesc&& InBindingDesc)
+{
+	// Record this binding in the valid binding map. This will be used to then remove all the invalid bindings from the collection
+	// The binding data view can be filled in with null memory, as only the struct is required for resolution
+	ValidBindingMap.Add(InBindingDesc.ID, FPropertyBindingDataView(InBindingDesc.Struct, nullptr));
+
+	Class->BindingCollection.BindingDescs.Add(MoveTemp(InBindingDesc));
+}
+
+void FBindingCompiler::AddRootBindingDesc()
+{
+	FSceneStateBindingDesc BindingDesc = Blueprint->CreateRootBinding();
+
+	// Set the Struct to the New Class during compilation
+	BindingDesc.Struct = Class;
+	BindingDesc.DataHandle = FSceneStateBindingDataHandle(ESceneStateDataType::Root);
+
+	AddBindingDesc(MoveTemp(BindingDesc));
+}
+
+void FBindingCompiler::AddStateMachineBindingDescs()
+{
+	for (const TPair<FGuid, uint16>& Pair : Class->StateMachineIdToIndex)
+	{
+		const uint16 StateMachineIndex = Pair.Value;
+
+		const FSceneStateMachine& StateMachine = Class->StateMachines[StateMachineIndex];
+		if (!StateMachine.Parameters.IsValid())
+		{
+			continue;
+		}
+
+		FSceneStateBindingDesc BindingDesc;
+		BindingDesc.ID = Pair.Key;
+		BindingDesc.Name = TEXT("State Machine Parameters");
+		BindingDesc.Struct = StateMachine.Parameters.GetPropertyBagStruct();
+		BindingDesc.DataHandle = FSceneStateBindingDataHandle(ESceneStateDataType::StateMachine, StateMachineIndex);
+
+		AddBindingDesc(MoveTemp(BindingDesc));
+	}
+}
+
+void FBindingCompiler::AddTransitionBindingDescs()
+{
+	for (const TPair<uint16, FInstancedPropertyBag>& Pair : Class->TransitionParameters)
+	{
+		const uint16 TransitionIndex = Pair.Key;
+
+		const FSceneStateTransitionMetadata& TransitionMetadata = Class->TransitionMetadata[TransitionIndex];
+
+		FSceneStateBindingDesc BindingDesc;
+		BindingDesc.ID = TransitionMetadata.ParametersId;
+		BindingDesc.Name = TEXT("Transition Parameters");
+		BindingDesc.Struct = Pair.Value.GetPropertyBagStruct();
+		BindingDesc.DataHandle = FSceneStateBindingDataHandle(ESceneStateDataType::Transition, TransitionIndex);
+
+		AddBindingDesc(MoveTemp(BindingDesc));
+	}
+}
+
+void FBindingCompiler::AddTaskBindingDescs()
+{
+	check(Class->Tasks.Num() == Class->TaskMetadata.Num());
+
+	TArray<FSceneStateBindingDesc>& BindingDescs = Class->BindingCollection.BindingDescs;
+
+	for (int32 TaskIndex = 0; TaskIndex < Class->Tasks.Num(); ++TaskIndex)
+	{
+		const FSceneStateTask& Task = Class->Tasks[TaskIndex].Get<FSceneStateTask>();
+		const FSceneStateTaskMetadata& TaskMetadata = Class->TaskMetadata[TaskIndex];
+		const FConstStructView TaskInstance = Class->TaskInstances[TaskIndex];
+
+		// Task Instance Binding
+		{
+			FSceneStateBindingDesc BindingDesc;
+			BindingDesc.ID = TaskMetadata.TaskId;
+			BindingDesc.Name = FName(TEXT("Task"), TaskIndex);
+			BindingDesc.Struct = TaskInstance.GetScriptStruct();
+			BindingDesc.DataHandle = FSceneStateBindingDataHandle(ESceneStateDataType::Task, TaskIndex);
+
+			AddBindingDesc(MoveTemp(BindingDesc));
+		}
+
+		const FSceneStateTaskBindingExtension* BindingExtension = Task.GetBindingExtension();
+		if (!BindingExtension)
+		{
+			continue;
+		}
+
+		TSet<uint16> DataIndices;
+		BindingExtension->VisitBindingDescs(TaskInstance,
+			[&BindingDescs, TaskIndex, &DataIndices, This = this](const FTaskBindingDesc& InBindingDesc)
+			{
+				bool bAlreadyInSet = false;
+				DataIndices.Add(InBindingDesc.DataIndex, &bAlreadyInSet);
+				checkf(!bAlreadyInSet, TEXT("Data Index already being used by another Desc in this Task!"));
+
+				FSceneStateBindingDesc BindingDesc;
+				BindingDesc.ID = InBindingDesc.Id;
+				BindingDesc.Name = InBindingDesc.Name;
+				BindingDesc.Struct = InBindingDesc.Struct;
+				BindingDesc.DataHandle = FSceneStateBindingDataHandle(ESceneStateDataType::TaskExtension, TaskIndex, InBindingDesc.DataIndex);
+
+				This->AddBindingDesc(MoveTemp(BindingDesc));
+			});
+	}
+}
+
+void FBindingCompiler::AddEventHandlerBindingDescs()
+{
+	for (int32 EventHandlerIndex = 0; EventHandlerIndex < Class->EventHandlers.Num(); ++EventHandlerIndex)
+	{
+		const FSceneStateEventHandler& EventHandler = Class->EventHandlers[EventHandlerIndex];
+
+		const USceneStateEventSchemaObject* EventSchema = EventHandler.GetEventSchemaHandle().GetEventSchema();
+		if (!EventSchema)
+		{
+			continue;
+		}
+
+		FSceneStateBindingDesc BindingDesc;
+		BindingDesc.ID = EventHandler.GetHandlerId();
+		BindingDesc.Name = EventSchema->Name;
+		BindingDesc.Struct = EventSchema->Struct;
+		BindingDesc.DataHandle = FSceneStateBindingDataHandle(ESceneStateDataType::EventHandler, EventHandlerIndex);
+
+		AddBindingDesc(MoveTemp(BindingDesc));
+	}
+}
+
+bool FBindingCompiler::ValidateBinding(const FSceneStateBinding& InBinding) const
+{
+	// Source must be accessible by the target struct.
+	// This mismatch could happen if copying an object bound to a scoped parameter and pasting it outside such scope
+	TArray<TInstancedStruct<FPropertyBindingBindableStructDescriptor>> AccessibleStructs;
+	Blueprint->GetBindableStructs(InBinding.GetTargetPath().GetStructID(), AccessibleStructs);
+
+	const FGuid& SourceStructId = InBinding.GetSourcePath().GetStructID();
+	const bool bSourceAccessible = AccessibleStructs.ContainsByPredicate(
+		[&SourceStructId](TConstStructView<FPropertyBindingBindableStructDescriptor> InBindableStruct)
+		{
+			return InBindableStruct.Get().ID == SourceStructId;
+		});
+
+	if (!bSourceAccessible)
+	{
+		const FText ErrorMessageFormat = FText::Format(LOCTEXT("InaccessibleSourceError"
+			, "Source '{0}' cannot be bound to target '{1}' because it's inaccessible (in another scope)")
+			, FText::FromString(InBinding.GetSourcePath().ToString())
+			, FText::FromString(InBinding.GetTargetPath().ToString()));
+
+		Context.MessageLog.Error(*ErrorMessageFormat.ToString());
+		return false;
+	}
+
+	if (!InBinding.SourceDataHandle.IsValid())
+	{
+		const FText ErrorMessageFormat = FText::Format(LOCTEXT("InvalidSourceHandleError"
+			, "Source '{0}' data handle was not found")
+			, FText::FromString(InBinding.GetSourcePath().ToString()));
+
+		Context.MessageLog.Error(*ErrorMessageFormat.ToString());
+		return false;
+	}
+
+	if (!InBinding.TargetDataHandle.IsValid())
+	{
+		const FText ErrorMessageFormat = FText::Format(LOCTEXT("InvalidTargetHandleError"
+			, "Target '{0}' data handle was not found")
+			, FText::FromString(InBinding.GetTargetPath().ToString()));
+
+		Context.MessageLog.Error(*ErrorMessageFormat.ToString());
+		return false;
+	}
+
+	return true;
+}
+
+void FBindingCompiler::ResolveBindingDataHandles()
+{
+	for (TArray<FSceneStateBinding>::TIterator Iter(Class->BindingCollection.Bindings); Iter; ++Iter)
+	{
+		FSceneStateBinding& Binding = *Iter;
+		Binding.SourceDataHandle = GetDataHandleById(Binding.GetSourcePath().GetStructID());
+		Binding.TargetDataHandle = GetDataHandleById(Binding.GetTargetPath().GetStructID());
+
+		if (!ValidateBinding(Binding))
+		{
+			Iter.RemoveCurrentSwap();
+		}
+	}
+}
+
+void FBindingCompiler::RemoveInvalidBindings()
+{
+	// Remove all the bindings not present in the valid binding map
+	Blueprint->BindingCollection.RemoveInvalidBindings(ValidBindingMap);
+	Class->BindingCollection.RemoveInvalidBindings(ValidBindingMap);
+}
+
+void FBindingCompiler::BatchCopies()
+{
+	// To batch first sort the bindings so all the same target data handles are together
+	Class->BindingCollection.Bindings.StableSort(
+		[](const FSceneStateBinding& InLeft, const FSceneStateBinding& InRight)
+		{
+			return InLeft.TargetDataHandle.AsNumber() < InRight.TargetDataHandle.AsNumber();
+		});
+
+	// Add 1 batch per group of bindings that all have the same target data handle
+	for (TArray<FSceneStateBinding>::TIterator BindingIter(Class->BindingCollection.Bindings); BindingIter;)
+	{
+		const FSceneStateBindingDataHandle TargetDataHandle = BindingIter->TargetDataHandle;
+
+		const FSceneStateBindingDesc* TargetDesc = Class->BindingCollection.FindBindingDesc(TargetDataHandle);
+		if (!ensure(TargetDesc))
+		{
+			++BindingIter;
+			continue;
+		}
+
+		const FPropertyBindingIndex16 BatchBindings(Class->BindingCollection.GetNumCopyBatches());
+
+		FPropertyBindingCopyInfoBatch& CopyBatch = Class->BindingCollection.AddCopyBatch();
+		CopyBatch.TargetStruct = TInstancedStruct<FSceneStateBindingDesc>::Make(*TargetDesc);
+		CopyBatch.BindingsBegin = FPropertyBindingIndex16(BindingIter.GetIndex());
+
+		// Get the last index that matches the current target data handle
+		for (; BindingIter && BindingIter->TargetDataHandle == TargetDataHandle; ++BindingIter)
+		{
+		}
+
+		// Binding end is one past the last binding.
+		CopyBatch.BindingsEnd = FPropertyBindingIndex16(BindingIter.GetIndex());
+		OnBindingsBatchCompiled(BatchBindings, TargetDataHandle);
+	}
+}
+
+void FBindingCompiler::OnBindingsBatchCompiled(FPropertyBindingIndex16 InBindingsBatch, const FSceneStateBindingDataHandle& InTargetDataHandle)
+{
+	// External Data Types as Targets are not supported by design.
+	check(!InTargetDataHandle.IsExternalDataType());
+
+	const ESceneStateDataType DataType = static_cast<ESceneStateDataType>(InTargetDataHandle.GetDataType());
+	const uint16 DataIndex = InTargetDataHandle.GetDataIndex();
+
+	if (DataType == ESceneStateDataType::Task)
+	{
+		FSceneStateTask& Task = Class->Tasks[DataIndex].Get<FSceneStateTask>();
+		Task.BindingsBatch = InBindingsBatch;
+	}
+	else if (DataType == ESceneStateDataType::TaskExtension)
+	{
+		FSceneStateTask& Task = Class->Tasks[DataIndex].Get<FSceneStateTask>();
+
+		FSceneStateTaskBindingExtension* BindingExtension = const_cast<FSceneStateTaskBindingExtension*>(Task.GetBindingExtension());
+		check(BindingExtension);
+		BindingExtension->SetBindingBatch(InTargetDataHandle.GetDataSubIndex(), InBindingsBatch.Get());
+	}
+	else if (DataType == ESceneStateDataType::StateMachine)
+	{
+		FSceneStateMachine& StateMachine = Class->StateMachines[DataIndex];
+		StateMachine.BindingsBatch = InBindingsBatch;
+	}
+	else if (DataType == ESceneStateDataType::Transition)
+	{
+		FSceneStateTransition& Transition = Class->Transitions[DataIndex];
+		Transition.BindingsBatch = InBindingsBatch;
+	}
+	else
+	{
+		// No other Data Types are supported as Targets at the moment
+		checkNoEntry();
+	}
+}
+
+FSceneStateBindingDataHandle FBindingCompiler::GetDataHandleById(const FGuid& InStructId)
+{
+	if (Blueprint->GetRootId() == InStructId)
+	{
+		return FSceneStateBindingDataHandle(ESceneStateDataType::Root);
+	}
+
+	// Tasks
+	for (int32 TaskIndex = 0; TaskIndex < Class->TaskMetadata.Num(); ++TaskIndex)
+	{
+		const FSceneStateTaskMetadata& TaskMetadata = Class->TaskMetadata[TaskIndex];
+		if (TaskMetadata.TaskId == InStructId)
+		{
+			return FSceneStateBindingDataHandle(ESceneStateDataType::Task, TaskIndex);
+		}
+	}
+
+	// State Machines
+	for (const TPair<FGuid, uint16>& Pair : Class->StateMachineIdToIndex)
+	{
+		if (Pair.Key == InStructId)
+		{
+			return FSceneStateBindingDataHandle(ESceneStateDataType::StateMachine, Pair.Value);
+		}
+	}
+
+	// Transitions
+	for (const TPair<uint16, FInstancedPropertyBag>& Pair : Class->TransitionParameters)
+	{
+		const uint16 TransitionIndex = Pair.Key;
+
+		const FSceneStateTransitionMetadata& TransitionMetadata = Class->TransitionMetadata[TransitionIndex];
+		if (TransitionMetadata.ParametersId == InStructId)
+		{
+			return FSceneStateBindingDataHandle(ESceneStateDataType::Transition, TransitionIndex);
+		}
+	}
+
+	// Event Handlers
+	for (int32 EventHandlerIndex = 0; EventHandlerIndex < Class->EventHandlers.Num(); ++EventHandlerIndex)
+	{
+		const FSceneStateEventHandler& EventHandler = Class->EventHandlers[EventHandlerIndex];
+		if (EventHandler.GetHandlerId() == InStructId)
+		{
+			return FSceneStateBindingDataHandle(ESceneStateDataType::EventHandler, EventHandlerIndex);
+		}
+	}
+
+	// Task Custom Data 
+	for (int32 TaskIndex = 0; TaskIndex < Class->Tasks.Num(); ++TaskIndex)
+	{
+		const FSceneStateTask& Task = Class->Tasks[TaskIndex].Get<const FSceneStateTask>();
+		const FStructView TaskInstance = Class->TaskInstances[TaskIndex];
+
+		FStructView DataView;
+		uint16 DataIndex;
+
+		const FSceneStateTaskBindingExtension* BindingExtension = Task.GetBindingExtension();
+		if (BindingExtension && BindingExtension->FindDataById(TaskInstance, InStructId, DataView, DataIndex))
+		{
+			return FSceneStateBindingDataHandle(ESceneStateDataType::TaskExtension, TaskIndex, DataIndex);
+		}
+	}
+
+	return FSceneStateBindingDataHandle();
+}
+
+} // UE::SceneState::Editor
+
+#undef LOCTEXT_NAMESPACE
diff --git a/Engine/Plugins/Experimental/SceneState/Source/SceneStateBlueprintEditor/Private/DebugControls/SceneStateDebugControlsObjectDetails.cpp b/Engine/Plugins/Experimental/SceneState/Source/SceneStateBlueprintEditor/Private/DebugControls/SceneStateDebugControlsObjectDetails.cpp
new file mode 100644
--- /dev/null
+++ b/Engine/Plugins/Experimental/SceneState/Source/SceneStateBlueprintEditor/Private/DebugControls/SceneStateDebugControlsObjectDetails.cpp
@@ -0,0 +1,84 @@
+// Copyright Epic Games, Inc. All Rights Reserved.
+
+#include "SceneStateDebugControlsObjectDetails.h"
+#include "DetailCategoryBuilder.h"
+#include "DetailLayoutBuilder.h"
+#include "DetailWidgetRow.h"
+#include "IDetailChildrenBuilder.h"
+#include "PropertyCustomizationHelpers.h"
+#include "SceneStateDebugControlsObject.h"
+#include "SceneStateEventUtils.h"
+#include "SceneStateObject.h"
+#include "Widgets/Input/SButton.h"
+#include "Widgets/SBoxPanel.h"
+#include "Widgets/Text/STextBlock.h"
+
+#define LOCTEXT_NAMESPACE "SceneStateDebugControlsObjectDetails"
+
+namespace UE::SceneState::Editor
+{
+
+void FDebugControlsObjectDetails::CustomizeDetails(IDetailLayoutBuilder& InDetailBuilder)
+{
+	DebugControls = InDetailBuilder.GetObjectsOfTypeBeingCustomized<USceneStateDebugControlsObject>();
+	CustomizeEventDetails(InDetailBuilder);
+	CustomizeDebuggedObjectDetails(InDetailBuilder);
+}
+
+void FDebugControlsObjectDetails::CustomizeEventDetails(IDetailLayoutBuilder& InDetailBuilder)
+{
+	EventsHandle = InDetailBuilder.GetProperty(GET_MEMBER_NAME_CHECKED(USceneStateDebugControlsObject, Events));
+	EventsHandle->MarkHiddenByCustomization();
+
+	IDetailCategoryBuilder& EventsCategory = InDetailBuilder.EditCategory(TEXT("Events"));
+	EventsCategory.AddProperty(EventsHandle);
+}
+
+USceneStateObject* FDebugControlsObjectDetails::GetDebuggedObject(int32 InIndex) const
+{
+	if (!DebugControls.IsValidIndex(InIndex))
+	{
+		return nullptr;
+	}
+
+	USceneStateDebugControlsObject* const DebugControlsObject = DebugControls[InIndex].Get();
+	if (!DebugControlsObject)
+	{
+		return nullptr;
+	}
+
+	return DebugControlsObject->DebuggedObjectWeak.Get();
+}
+
+void FDebugControlsObjectDetails::CustomizeDebuggedObjectDetails(IDetailLayoutBuilder& InDetailBuilder)
+{
+	TArray<UObject*> DebuggedObjects;
+	DebuggedObjects.Reserve(DebugControls.Num());
+
+	// Gather all the Debugged objects within the debug controls
+	Algo::Transform(DebugControls, DebuggedObjects,
+		[](const TWeakObjectPtr<USceneStateDebugControlsObject>& InDebugControls)->UObject*
+		{
+			if (USceneStateDebugControlsObject* DebugControl = InDebugControls.Get())
+			{
+				return DebugControl->DebuggedObjectWeak.Get();
+			}
+			return nullptr;
+		});
+
+	DebuggedObjects.RemoveAll(
+		[](UObject* InObject)
+		{
+			return !InObject;
+		});
+
+	InDetailBuilder
+		.EditCategory(TEXT("Debugged Object"))
+		.AddExternalObjects(DebuggedObjects, EPropertyLocation::Default,
+			FAddPropertyParams()
+			.HideRootObjectNode(/*bHideRootObjectNode*/true));
+}
+
+} // UE::SceneState::Editor
+
+#undef LOCTEXT_NAMESPACE
diff --git a/Engine/Plugins/Experimental/SceneState/Source/SceneStateBlueprintEditor/Private/DebugExecutors/SceneStateDebugTaskExecutor.cpp b/Engine/Plugins/Experimental/SceneState/Source/SceneStateBlueprintEditor/Private/DebugExecutors/SceneStateDebugTaskExecutor.cpp
new file mode 100644
--- /dev/null
+++ b/Engine/Plugins/Experimental/SceneState/Source/SceneStateBlueprintEditor/Private/DebugExecutors/SceneStateDebugTaskExecutor.cpp
@@ -0,0 +1,105 @@
+// Copyright Epic Games, Inc. All Rights Reserved.
+
+#include "SceneStateDebugTaskExecutor.h"
+#include "Nodes/SceneStateMachineStateNode.h"
+#include "Nodes/SceneStateMachineTaskNode.h"
+#include "SceneStateGeneratedClass.h"
+#include "SceneStateMachineGraphSchema.h"
+#include "Tasks/SceneStateTaskInstance.h"
+
+namespace UE::SceneState::Editor
+{
+
+const FSceneStateTask* FDebugTaskExecutor::GetTask(const FSceneStateExecutionContext& InExecutionContext) const
+{
+	if (const USceneStateGeneratedClass* GeneratedClass = InExecutionContext.GetGeneratedClass())
+	{
+		return GeneratedClass->FindTaskFromNode(GetNodeKey());
+	}
+	return nullptr;
+}
+
+const FSceneState* FDebugTaskExecutor::GetState(const USceneStateGeneratedClass* InGeneratedClass) const
+{
+	if (const USceneStateMachineTaskNode* const TaskNode = Cast<USceneStateMachineTaskNode>(GetNodeKey().ResolveObjectPtr()))
+	{
+		if (const USceneStateMachineStateNode* const StateNode = USceneStateMachineGraphSchema::FindConnectedStateNode(TaskNode))
+		{
+			return InGeneratedClass->FindStateFromNode(StateNode);
+		}
+	}
+	return nullptr;
+}
+
+void FDebugTaskExecutor::Setup(const FSceneStateExecutionContext& InExecutionContext, const FSceneStateTask& InTask)
+{
+	const USceneStateGeneratedClass* GeneratedClass = InExecutionContext.GetGeneratedClass();
+	if (!GeneratedClass)
+	{
+		return;
+	}
+
+	const FSceneState* ParentState = GetState(GeneratedClass);
+	if (!ParentState)
+	{
+		return;
+	}
+
+	const FSceneStateRange TaskRange = ParentState->GetTaskRange(); 
+
+	const uint16 TaskAbsoluteIndex = InTask.GetTaskIndex();
+	const uint16 TaskRelativeIndex = TaskAbsoluteIndex - TaskRange.Index;
+
+	// Make an empty task instance array that matches the state task count,
+	// and only set the task instance template for the task this executor cares about
+	TArray<FConstStructView> TemplateTaskInstances;
+	TemplateTaskInstances.SetNum(TaskRange.Count);
+	TemplateTaskInstances[TaskRelativeIndex] = GeneratedClass->GetTaskInstances()[TaskAbsoluteIndex];
+
+	ParentState->AllocateTaskInstances(InExecutionContext, TemplateTaskInstances);
+}
+
+void FDebugTaskExecutor::OnStart(const FSceneStateExecutionContext& InExecutionContext)
+{
+	if (const FSceneStateTask* Task = GetTask(InExecutionContext))
+	{
+		Setup(InExecutionContext, *Task);
+
+		const FStructView TaskInstance = Task->FindTaskInstance(InExecutionContext);
+
+		Task->Setup(InExecutionContext, TaskInstance);
+		Task->Start(InExecutionContext, TaskInstance);
+		ConditionallyExit(TaskInstance);
+	}
+}
+
+void FDebugTaskExecutor::OnTick(const FSceneStateExecutionContext& InExecutionContext, float InDeltaSeconds)
+{
+	if (const FSceneStateTask* Task = GetTask(InExecutionContext))
+	{
+		const FStructView TaskInstance = Task->FindTaskInstance(InExecutionContext);
+
+		Task->Tick(InExecutionContext, TaskInstance, InDeltaSeconds);
+		ConditionallyExit(TaskInstance);
+	}
+}
+
+void FDebugTaskExecutor::OnExit(const FSceneStateExecutionContext& InExecutionContext)
+{
+	if (const FSceneStateTask* Task = GetTask(InExecutionContext))
+	{
+		const FStructView TaskInstance = Task->FindTaskInstance(InExecutionContext);
+		Task->Stop(InExecutionContext, TaskInstance, ESceneStateTaskStopReason::State);
+	}
+}
+
+void FDebugTaskExecutor::ConditionallyExit(FStructView InTaskInstance)
+{
+	const FSceneStateTaskInstance* TaskInstance = InTaskInstance.GetPtr<FSceneStateTaskInstance>();
+	if (TaskInstance && TaskInstance->Status == EExecutionStatus::Finished)
+	{
+		Exit();
+	}
+}
+
+} // UE::SceneState::Editor
