Motion Design: Added motion design scene state plugin

diff --git a/Engine/Plugins/Experimental/SceneState/Source/SceneStateBlueprintEditor/Public/SceneStateBlueprintEditorUtils.h b/Engine/Plugins/Experimental/SceneState/Source/SceneStateBlueprintEditor/Public/SceneStateBlueprintEditorUtils.h
new file mode 100644
--- /dev/null
+++ b/Engine/Plugins/Experimental/SceneState/Source/SceneStateBlueprintEditor/Public/SceneStateBlueprintEditorUtils.h
@@ -0,0 +1,67 @@
+// Copyright Epic Games, Inc. All Rights Reserved.
+
+#pragma once
+
+#include "PropertyEditorModule.h"
+#include "Templates/SharedPointerFwd.h"
+
+class IDetailChildrenBuilder;
+class IPropertyHandle;
+class UClass;
+struct FGuid;
+struct FInstancedPropertyBag;
+
+namespace UE::SceneState::Editor
+{
+	/** Finds the Task Id for a given Property Handle by first finding the outer Task Node and returning its Task id. */
+	SCENESTATEBLUEPRINTEDITOR_API FGuid FindTaskId(const TSharedRef<IPropertyHandle>& InPropertyHandle);
+
+	/**
+	 * Returns whether a given property is an object property of a given class
+	 * @param InProperty the property to check
+	 * @param InClass the class to check
+	 * @return true if the property is an object property of the given class
+	 */
+	SCENESTATEBLUEPRINTEDITOR_API bool IsObjectPropertyOfClass(const FProperty* InProperty, const UClass* InClass);
+
+	/**
+	 * Returns whether a given property is of a given struct type
+	 * @param InProperty the property to check
+	 * @param InStruct the type to check
+	 * @return true if the property is a struct property of the given struct type  
+	 */
+	SCENESTATEBLUEPRINTEDITOR_API bool IsStruct(const FProperty* InProperty, const UScriptStruct* InStruct);
+
+	/**
+	 * Gets the Guid value from the Property Handle.
+	 * @param InGuidPropertyHandle the guid property handle to get the guid value from
+	 * @param OutGuid the returned guid if the operation was successful
+	 * @return whether the guid was obtained successfully, or whether it failed or had multiple values
+	 */
+	SCENESTATEBLUEPRINTEDITOR_API FPropertyAccess::Result GetGuid(const TSharedRef<IPropertyHandle>& InGuidPropertyHandle, FGuid& OutGuid);
+
+	/**
+	 * Tries to create a structure data provider for the given instanced struct property handle
+	 * @param InStructHandle the instanced struct property handle to create a struct data provider for
+	 * @return the created struct data provider
+	 */
+	SCENESTATEBLUEPRINTEDITOR_API TSharedRef<IStructureDataProvider> CreateInstancedStructDataProvider(const TSharedRef<IPropertyHandle>& InStructHandle);
+
+	/**
+	 * Compares the structure layout of two property bags
+	 * @param InParametersA the first property bag comparand
+	 * @param InParametersB the second property bag comparand
+	 * @return true if both have the same number of properties with same names and compatible types.
+	 */
+	SCENESTATEBLUEPRINTEDITOR_API bool CompareParametersLayout(const FInstancedPropertyBag& InParametersA, const FInstancedPropertyBag& InParametersB);
+
+	/** Sets the Instance Meta-data of a given property handle to a Task Id guid */
+	SCENESTATEBLUEPRINTEDITOR_API void AssignBindingId(const TSharedRef<IPropertyHandle>& InPropertyHandle, const FGuid& InTaskId);
+
+	/** Finds the Common Class for all the Objects within the Object Property Handle */
+	SCENESTATEBLUEPRINTEDITOR_API UClass* FindCommonBase(const TSharedRef<IPropertyHandle>& InPropertyHandle);
+
+	/** Helper function to add the Common properties of the object property handle to the Child Builder */
+	SCENESTATEBLUEPRINTEDITOR_API bool AddObjectProperties(const TSharedRef<IPropertyHandle>& InPropertyHandle, IDetailChildrenBuilder& InChildBuilder);
+
+} // UE::SceneState::Editor
diff --git a/Engine/Plugins/Experimental/SceneState/Source/SceneStateEvent/Public/SceneStateEventStream.h b/Engine/Plugins/Experimental/SceneState/Source/SceneStateEvent/Public/SceneStateEventStream.h
new file mode 100644
--- /dev/null
+++ b/Engine/Plugins/Experimental/SceneState/Source/SceneStateEvent/Public/SceneStateEventStream.h
@@ -0,0 +1,59 @@
+// Copyright Epic Games, Inc. All Rights Reserved.
+
+#pragma once
+
+#include "Containers/Array.h"
+#include "Containers/Map.h"
+#include "SceneStateEvent.h"
+#include "UObject/Object.h"
+#include "SceneStateEventStream.generated.h"
+
+struct FSceneStateEventHandler;
+struct FSceneStateEventSchemaHandle;
+
+/** Holds and keeps track of Events added at Runtime */
+UCLASS(MinimalAPI)
+class USceneStateEventStream : public UObject
+{
+	GENERATED_BODY()
+
+public:
+	/** Registers the Event Stream to the Event Subsystem to listen to Broadcast Events */
+	SCENESTATEEVENT_API bool Register();
+
+	/** Registers the Event Stream from the Event Subsystem */
+	SCENESTATEEVENT_API void Unregister();
+
+	/** Pushes a new event into the stream */
+	SCENESTATEEVENT_API void PushEvent(FSharedStruct&& InEvent);
+
+	/** Pushes a new event into the stream */
+	SCENESTATEEVENT_API void PushEvent(const FSharedStruct& InEvent);
+
+	/** Consumes the first (oldest) Event that was pushed that matches the given Schema */
+	SCENESTATEEVENT_API bool ConsumeEventBySchema(const FSceneStateEventSchemaHandle& InEventSchemaHandle);
+
+	/** Finds the first (oldest) Event that was pushed that matches the given Schema */
+	SCENESTATEEVENT_API const FSceneStateEvent* FindEventBySchema(const FSceneStateEventSchemaHandle& InEventSchemaHandle) const;
+
+	/** Finds the event that is captured by the Handler with the given Id */
+	SCENESTATEEVENT_API FSceneStateEvent* FindCapturedEvent(const FGuid& InHandlerId);
+
+	/** Removes the Events (first/oldest) that match the given Event Handlers and moves them to the Captured Event map */
+	SCENESTATEEVENT_API void CaptureEvents(TConstArrayView<FSceneStateEventHandler> InEventHandlers);
+
+	/** Cleanup the Captured Events that match the given Event Handlers */
+	SCENESTATEEVENT_API void ResetCapturedEvents(TConstArrayView<FSceneStateEventHandler> InEventHandlers);
+
+private:
+	/** Returns the index of the first event matching the event schema */
+	int32 GetEventIndexBySchema(const FSceneStateEventSchemaHandle& InEventSchemaHandle) const;
+
+	/** Active Events kept in push order */
+	UPROPERTY(Transient, DuplicateTransient, TextExportTransient)
+	TArray<FSharedStruct> Events;
+
+	/** Map of the Handler Id to the Event it has captured */
+	UPROPERTY(Transient, DuplicateTransient, TextExportTransient)
+	TMap<FGuid, FSharedStruct> CapturedEvents;
+};
diff --git a/Engine/Plugins/Experimental/SceneState/Source/SceneStateEvent/Public/SceneStateEventUtils.h b/Engine/Plugins/Experimental/SceneState/Source/SceneStateEvent/Public/SceneStateEventUtils.h
new file mode 100644
--- /dev/null
+++ b/Engine/Plugins/Experimental/SceneState/Source/SceneStateEvent/Public/SceneStateEventUtils.h
@@ -0,0 +1,46 @@
+// Copyright Epic Games, Inc. All Rights Reserved.
+
+#pragma once
+
+#include "HAL/Platform.h"
+
+class USceneStateEventStream;
+struct FInstancedStruct;
+struct FSceneStateEventSchemaHandle;
+struct FSceneStateEventTemplate;
+
+namespace UE::SceneState
+{
+	/**
+	 * Pushes an event to the given event stream
+	 * @param InEventStream the stream to push the event to
+	 * @param InEventSchemaHandle handle to the event schema used to create the event
+	 * @param InEventData the event payload. Must match the event schema struct
+	 * @return true if the event was pushed successfully, false otherwise
+	 */
+	SCENESTATEEVENT_API bool PushEvent(USceneStateEventStream* InEventStream, const FSceneStateEventSchemaHandle& InEventSchemaHandle, FInstancedStruct&& InEventData);
+
+	/**
+	 * Pushes an event to the given event stream
+	 * @param InEventStream the stream to push the event to
+	 * @param InEventTemplate the template used to create the event
+	 * @return true if the event was pushed successfully, false otherwise
+	 */
+	SCENESTATEEVENT_API bool PushEvent(USceneStateEventStream* InEventStream, const FSceneStateEventTemplate& InEventTemplate);
+
+	/**
+	 * Broadcasts an event to the every registered event stream
+	 * @param InEventSchemaHandle handle to the event schema used to create the event
+	 * @param InEventData the event payload. Must match the event schema struct
+	 * @return true if the event was broadcasted successfully, false otherwise
+	 */
+	SCENESTATEEVENT_API bool BroadcastEvent(const FSceneStateEventSchemaHandle& InEventSchemaHandle, FInstancedStruct&& InEventData);
+
+	/**
+	 * Broadcasts an event to the every registered event stream
+	 * @param InEventTemplate the template used to create the event
+	 * @return true if the event was broadcasted successfully, false otherwise
+	 */
+	SCENESTATEEVENT_API bool BroadcastEvent(const FSceneStateEventTemplate& InEventTemplate);
+
+} // UE::SceneState
diff --git a/Engine/Plugins/Experimental/SceneState/Source/SceneStateGameplay/Public/SceneStateActor.h b/Engine/Plugins/Experimental/SceneState/Source/SceneStateGameplay/Public/SceneStateActor.h
new file mode 100644
--- /dev/null
+++ b/Engine/Plugins/Experimental/SceneState/Source/SceneStateGameplay/Public/SceneStateActor.h
@@ -0,0 +1,31 @@
+// Copyright Epic Games, Inc. All Rights Reserved.
+
+#pragma once
+
+#include "GameFramework/Actor.h"
+#include "SceneStateActor.generated.h"
+
+class USceneStateComponent;
+class USceneStateObject;
+
+UCLASS(MinimalAPI)
+class ASceneStateActor : public AActor
+{
+	GENERATED_BODY()
+
+public:
+	SCENESTATEGAMEPLAY_API static const FLazyName SceneStateComponentName;
+
+	SCENESTATEGAMEPLAY_API ASceneStateActor(const FObjectInitializer& InObjectInitializer);
+
+	SCENESTATEGAMEPLAY_API void SetSceneStateClass(TSubclassOf<USceneStateObject> InSceneStateClass);
+
+	SCENESTATEGAMEPLAY_API TSubclassOf<USceneStateObject> GetSceneStateClass() const;
+
+	UFUNCTION(BlueprintCallable, Category = "Scene State")
+	SCENESTATEGAMEPLAY_API USceneStateObject* GetSceneState() const;
+
+private:
+	UPROPERTY(VisibleAnywhere, BlueprintReadOnly, Category="Scene State", meta=(AllowPrivateAccess="true"))
+	TObjectPtr<USceneStateComponent> SceneStateComponent;
+};
diff --git a/Engine/Plugins/Experimental/SceneState/Source/SceneStateGameplay/Public/SceneStateComponent.h b/Engine/Plugins/Experimental/SceneState/Source/SceneStateGameplay/Public/SceneStateComponent.h
new file mode 100644
--- /dev/null
+++ b/Engine/Plugins/Experimental/SceneState/Source/SceneStateGameplay/Public/SceneStateComponent.h
@@ -0,0 +1,48 @@
+// Copyright Epic Games, Inc. All Rights Reserved.
+
+#pragma once
+
+#include "Components/ActorComponent.h"
+#include "SceneStateComponent.generated.h"
+
+class USceneStateObject;
+class USceneStateComponentPlayer;
+struct FSceneStateComponentInstanceData;
+
+UCLASS(MinimalAPI, BlueprintType, meta=(BlueprintSpawnableComponent))
+class USceneStateComponent : public UActorComponent
+{
+	GENERATED_BODY()
+
+public:
+	SCENESTATEGAMEPLAY_API static const FLazyName SceneStatePlayerName;
+
+	SCENESTATEGAMEPLAY_API USceneStateComponent(const FObjectInitializer& InObjectInitializer);
+
+	USceneStateComponentPlayer* GetSceneStatePlayer() const
+	{
+		return SceneStatePlayer;
+	}
+
+	SCENESTATEGAMEPLAY_API TSubclassOf<USceneStateObject> GetSceneStateClass() const;
+
+	SCENESTATEGAMEPLAY_API void SetSceneStateClass(TSubclassOf<USceneStateObject> InSceneStateClass);
+
+	UFUNCTION(BlueprintCallable, Category = "Scene State")
+	SCENESTATEGAMEPLAY_API USceneStateObject* GetSceneState() const;
+
+	void ApplyComponentInstanceData(FSceneStateComponentInstanceData* InComponentInstanceData);
+
+	//~ Begin UActorComponent
+	SCENESTATEGAMEPLAY_API virtual void InitializeComponent() override;
+	SCENESTATEGAMEPLAY_API virtual void OnRegister() override;
+	SCENESTATEGAMEPLAY_API virtual void BeginPlay() override;
+	SCENESTATEGAMEPLAY_API virtual void TickComponent(float InDeltaTime, ELevelTick InTickType, FActorComponentTickFunction* InThisTickFunction) override;
+	SCENESTATEGAMEPLAY_API virtual void EndPlay(const EEndPlayReason::Type InEndPlayReason) override;
+	SCENESTATEGAMEPLAY_API virtual TStructOnScope<FActorComponentInstanceData> GetComponentInstanceData() const override;
+	//~ End UActorComponent
+
+private:
+	UPROPERTY(VisibleAnywhere, Instanced, Category="Scene State", meta=(ShowInnerProperties))
+	TObjectPtr<USceneStateComponentPlayer> SceneStatePlayer;
+};
diff --git a/Engine/Plugins/Experimental/SceneState/Source/SceneStateGameplay/Public/SceneStateComponentPlayer.h b/Engine/Plugins/Experimental/SceneState/Source/SceneStateGameplay/Public/SceneStateComponentPlayer.h
new file mode 100644
--- /dev/null
+++ b/Engine/Plugins/Experimental/SceneState/Source/SceneStateGameplay/Public/SceneStateComponentPlayer.h
@@ -0,0 +1,23 @@
+// Copyright Epic Games, Inc. All Rights Reserved.
+
+#pragma once
+
+#include "SceneStatePlayer.h"
+#include "SceneStateComponentPlayer.generated.h"
+
+class USceneStateComponent;
+
+/** Scene State Player for Scene State Components */
+UCLASS(MinimalAPI)
+class USceneStateComponentPlayer : public USceneStatePlayer
+{
+	GENERATED_BODY()
+
+protected:
+	SCENESTATEGAMEPLAY_API AActor* GetActor() const;
+
+	//~ Begin USceneStatePlayer
+	SCENESTATEGAMEPLAY_API virtual bool OnGetContextName(FString& OutContextName) const override;
+	SCENESTATEGAMEPLAY_API virtual bool OnGetContextObject(UObject*& OutContextObject) const override;
+	//~ End USceneStatePlayer
+};
diff --git a/Engine/Plugins/Experimental/SceneState/Source/SceneStateMachineGraph/Private/Nodes/SceneStateMachineNodeNameValidator.cpp b/Engine/Plugins/Experimental/SceneState/Source/SceneStateMachineGraph/Private/Nodes/SceneStateMachineNodeNameValidator.cpp
new file mode 100644
--- /dev/null
+++ b/Engine/Plugins/Experimental/SceneState/Source/SceneStateMachineGraph/Private/Nodes/SceneStateMachineNodeNameValidator.cpp
@@ -0,0 +1,49 @@
+// Copyright Epic Games, Inc. All Rights Reserved.
+
+#include "SceneStateMachineNodeNameValidator.h"
+#include "Nodes/SceneStateMachineNode.h"
+#include "SceneStateMachineGraph.h"
+
+namespace UE::SceneState::Graph
+{
+
+FStateMachineNodeNameValidator::FStateMachineNodeNameValidator(const USceneStateMachineNode* InNode)
+{
+	check(InNode);
+
+	TArray<USceneStateMachineNode*> Nodes;
+
+	USceneStateMachineGraph* Graph = CastChecked<USceneStateMachineGraph>(InNode->GetOuter());
+	Graph->GetNodesOfClass<USceneStateMachineNode>(Nodes);
+
+	Names.Reserve(Nodes.Num());
+	for (USceneStateMachineNode* Node : Nodes)
+	{
+		if (Node != InNode)
+		{
+			Names.Add(Node->GetNodeName());
+		}
+	}
+}
+
+EValidatorResult FStateMachineNodeNameValidator::IsValid(const FName& InName, bool bInOriginal)
+{
+	if (InName.IsNone())
+	{
+		return EValidatorResult::EmptyName;
+	}
+
+	if (Names.Contains(InName))
+	{
+		return EValidatorResult::AlreadyInUse;
+	}
+
+	return EValidatorResult::Ok;
+}
+
+EValidatorResult FStateMachineNodeNameValidator::IsValid(const FString& InName, bool bInOriginal)
+{
+	return IsValid(FName(InName), bInOriginal);
+}
+
+} // UE::SceneState::Graph
diff --git a/Engine/Plugins/Experimental/SceneState/Source/SceneStateMachineGraph/Public/SceneStateMachineGraph.h b/Engine/Plugins/Experimental/SceneState/Source/SceneStateMachineGraph/Public/SceneStateMachineGraph.h
new file mode 100644
--- /dev/null
+++ b/Engine/Plugins/Experimental/SceneState/Source/SceneStateMachineGraph/Public/SceneStateMachineGraph.h
@@ -0,0 +1,51 @@
+// Copyright Epic Games, Inc. All Rights Reserved.
+
+#pragma once
+
+#include "EdGraph/EdGraph.h"
+#include "SceneStateEnums.h"
+#include "StructUtils/PropertyBag.h"
+#include "SceneStateMachineGraph.generated.h"
+
+class USceneStateMachineEntryNode;
+
+UCLASS(MinimalAPI, HideCategories=(Graph))
+class USceneStateMachineGraph : public UEdGraph
+{
+	GENERATED_BODY()
+
+public:
+	using FOnParametersChanged = TMulticastDelegate<void(USceneStateMachineGraph&)>;
+
+	USceneStateMachineGraph();
+
+	SCENESTATEMACHINEGRAPH_API static FOnParametersChanged::RegistrationType& OnParametersChanged();
+
+	SCENESTATEMACHINEGRAPH_API void NotifyParametersChanged();
+
+	SCENESTATEMACHINEGRAPH_API USceneStateMachineEntryNode* GetEntryNode() const;
+
+	//~ Begin UObject
+	SCENESTATEMACHINEGRAPH_API virtual void PostInitProperties() override;
+	SCENESTATEMACHINEGRAPH_API virtual void PostLoad() override;
+	SCENESTATEMACHINEGRAPH_API virtual void PostDuplicate(bool bInDuplicateForPIE) override;
+	SCENESTATEMACHINEGRAPH_API virtual void PostEditImport() override;
+	//~ End UObject
+
+	/** Called to set a new unique id for parameters (e.g. after duplicating) */
+	void GenerateNewParametersId();
+
+	/** Identifier for the Parameters Struct Id */
+	UPROPERTY(VisibleAnywhere, Category="Scene State")
+	FGuid ParametersId;
+
+	UPROPERTY(EditAnywhere, Category="Scene State")
+	FInstancedPropertyBag Parameters;
+
+	/** The run-mode for the State Machine. Currently only applies to Top-Level State Machines */
+	UPROPERTY(EditAnywhere, Category="Scene State")
+	ESceneStateMachineRunMode RunMode = ESceneStateMachineRunMode::Auto;
+
+private:
+	static FOnParametersChanged OnParametersChangedDelegate;
+};
diff --git a/Engine/Plugins/Experimental/SceneState/Source/SceneStateMachineGraph/Public/SceneStateMachineGraphUtils.h b/Engine/Plugins/Experimental/SceneState/Source/SceneStateMachineGraph/Public/SceneStateMachineGraphUtils.h
new file mode 100644
--- /dev/null
+++ b/Engine/Plugins/Experimental/SceneState/Source/SceneStateMachineGraph/Public/SceneStateMachineGraphUtils.h
@@ -0,0 +1,15 @@
+// Copyright Epic Games, Inc. All Rights Reserved.
+
+#pragma once
+
+#include "HAL/Platform.h"
+
+class UEdGraph;
+
+namespace UE::SceneState::Graph
+{
+	SCENESTATEMACHINEGRAPH_API bool CanDirectlyRemoveGraph(UEdGraph* InGraph);
+
+	SCENESTATEMACHINEGRAPH_API void RemoveGraph(UEdGraph* InGraph);
+
+} // UE::SceneState::Graph
