Motion Design: fixed scene state and data link being outside of the motion design category in content browser. In addition, updated the display names of scene state and data link assets to have the Motion Design prefix.

diff --git a/Engine/Plugins/Experimental/DataLink/Source/DataLink/Public/DataLinkGraph.h b/Engine/Plugins/Experimental/DataLink/Source/DataLink/Public/DataLinkGraph.h
--- a/Engine/Plugins/Experimental/DataLink/Source/DataLink/Public/DataLinkGraph.h
+++ b/Engine/Plugins/Experimental/DataLink/Source/DataLink/Public/DataLinkGraph.h
@@ -14,7 +14,7 @@ struct FDataLinkPinReference;
 class UEdGraph;
 #endif
 
-UCLASS(MinimalAPI, BlueprintType)
+UCLASS(MinimalAPI, BlueprintType, DisplayName="Motion Design Data Graph")
 class UDataLinkGraph : public UObject
 {
 	GENERATED_BODY()
diff --git a/Engine/Plugins/Experimental/DataLink/Source/DataLinkEditor/Private/DataLinkGraphFactory.cpp b/Engine/Plugins/Experimental/DataLink/Source/DataLinkEditor/Private/DataLinkGraphFactory.cpp
--- a/Engine/Plugins/Experimental/DataLink/Source/DataLinkEditor/Private/DataLinkGraphFactory.cpp
+++ b/Engine/Plugins/Experimental/DataLink/Source/DataLinkEditor/Private/DataLinkGraphFactory.cpp
@@ -1,6 +1,7 @@
 // Copyright Epic Games, Inc. All Rights Reserved.
 
 #include "DataLinkGraphFactory.h"
+#include "AssetToolsModule.h"
 #include "AssetTypeCategories.h"
 #include "DataLinkEdGraph.h"
 #include "DataLinkEdGraphSchema.h"
@@ -15,9 +16,21 @@ UDataLinkGraphFactory::UDataLinkGraphFactory()
 	bEditAfterNew = true;
 }
 
+FText UDataLinkGraphFactory::GetDisplayName() const
+{
+	return SupportedClass ? SupportedClass->GetDisplayNameText() : Super::GetDisplayName();
+}
+
+FString UDataLinkGraphFactory::GetDefaultNewAssetName() const
+{
+	// Short name removing "Motion Design" and "DataLink" prefix for new assets
+	return TEXT("NewDataGraph");
+}
+
 uint32 UDataLinkGraphFactory::GetMenuCategories() const
 {
-	return EAssetTypeCategories::Misc;
+	IAssetTools& AssetTools = FAssetToolsModule::GetModule().Get();
+	return AssetTools.FindAdvancedAssetCategory("MotionDesignCategory");
 }
 
 UObject* UDataLinkGraphFactory::FactoryCreateNew(UClass* InClass, UObject* InParent, FName InName, EObjectFlags InFlags, UObject* InContext, FFeedbackContext* InWarn)
diff --git a/Engine/Plugins/Experimental/DataLink/Source/DataLinkEditor/Private/DataLinkGraphFactory.h b/Engine/Plugins/Experimental/DataLink/Source/DataLinkEditor/Private/DataLinkGraphFactory.h
--- a/Engine/Plugins/Experimental/DataLink/Source/DataLinkEditor/Private/DataLinkGraphFactory.h
+++ b/Engine/Plugins/Experimental/DataLink/Source/DataLinkEditor/Private/DataLinkGraphFactory.h
@@ -14,6 +14,8 @@ public:
 	UDataLinkGraphFactory();
 
 	//~ Begin UFactory
+	virtual FText GetDisplayName() const override;
+	virtual FString GetDefaultNewAssetName() const override;
 	virtual uint32 GetMenuCategories() const override;
 	virtual UObject* FactoryCreateNew(UClass* InClass, UObject* InParent, FName InName, EObjectFlags InFlags, UObject* InContext, FFeedbackContext* InWarn) override;
 	//~ End UFactory
diff --git a/Engine/Plugins/Experimental/DataLinkOAuth/Source/DataLinkOAuth/Public/DataLinkOAuthDefaultSettings.h b/Engine/Plugins/Experimental/DataLinkOAuth/Source/DataLinkOAuth/Public/DataLinkOAuthDefaultSettings.h
--- a/Engine/Plugins/Experimental/DataLinkOAuth/Source/DataLinkOAuth/Public/DataLinkOAuthDefaultSettings.h
+++ b/Engine/Plugins/Experimental/DataLinkOAuth/Source/DataLinkOAuth/Public/DataLinkOAuthDefaultSettings.h
@@ -5,7 +5,7 @@
 #include "DataLinkOAuthSettings.h"
 #include "DataLinkOAuthDefaultSettings.generated.h"
 
-UCLASS(MinimalAPI)
+UCLASS(MinimalAPI, DisplayName="OAuth Default Settings")
 class UDataLinkOAuthDefaultSettings : public UDataLinkOAuthSettings
 {
 	GENERATED_BODY()
diff --git a/Engine/Plugins/Experimental/DataLinkOAuth/Source/DataLinkOAuth/Public/DataLinkOAuthSettings.h b/Engine/Plugins/Experimental/DataLinkOAuth/Source/DataLinkOAuth/Public/DataLinkOAuthSettings.h
--- a/Engine/Plugins/Experimental/DataLinkOAuth/Source/DataLinkOAuth/Public/DataLinkOAuthSettings.h
+++ b/Engine/Plugins/Experimental/DataLinkOAuth/Source/DataLinkOAuth/Public/DataLinkOAuthSettings.h
@@ -12,7 +12,7 @@ struct FDataLinkOAuthToken;
 struct FHttpServerRequest;
 struct FInstancedStruct;
 
-UCLASS(MinimalAPI, Abstract)
+UCLASS(MinimalAPI, Abstract, DisplayName="Motion Design Data OAuth Settings")
 class UDataLinkOAuthSettings : public UObject
 {
 	GENERATED_BODY()
diff --git a/Engine/Plugins/Experimental/DataLinkOAuth/Source/DataLinkOAuthEditor/Private/DataLinkOAuthSettingsFactory.cpp b/Engine/Plugins/Experimental/DataLinkOAuth/Source/DataLinkOAuthEditor/Private/DataLinkOAuthSettingsFactory.cpp
--- a/Engine/Plugins/Experimental/DataLinkOAuth/Source/DataLinkOAuthEditor/Private/DataLinkOAuthSettingsFactory.cpp
+++ b/Engine/Plugins/Experimental/DataLinkOAuth/Source/DataLinkOAuthEditor/Private/DataLinkOAuthSettingsFactory.cpp
@@ -1,6 +1,7 @@
 // Copyright Epic Games, Inc. All Rights Reserved.
 
 #include "DataLinkOAuthSettingsFactory.h"
+#include "AssetToolsModule.h"
 #include "AssetTypeCategories.h"
 #include "ClassViewerFilter.h"
 #include "ClassViewerModule.h"
@@ -44,9 +45,21 @@ UDataLinkOAuthSettingsFactory::UDataLinkOAuthSettingsFactory()
 	bEditAfterNew = true;
 }
 
+FText UDataLinkOAuthSettingsFactory::GetDisplayName() const
+{
+	return SupportedClass ? SupportedClass->GetDisplayNameText() : Super::GetDisplayName();
+}
+
+FString UDataLinkOAuthSettingsFactory::GetDefaultNewAssetName() const
+{
+	// Short name removing "Motion Design" and "DataLink" prefix for new assets
+	return TEXT("NewOAuthSettings");
+}
+
 uint32 UDataLinkOAuthSettingsFactory::GetMenuCategories() const
 {
-	return EAssetTypeCategories::Misc;
+	IAssetTools& AssetTools = FAssetToolsModule::GetModule().Get();
+	return AssetTools.FindAdvancedAssetCategory("MotionDesignCategory");
 }
 
 bool UDataLinkOAuthSettingsFactory::ConfigureProperties()
diff --git a/Engine/Plugins/Experimental/DataLinkOAuth/Source/DataLinkOAuthEditor/Private/DataLinkOAuthSettingsFactory.h b/Engine/Plugins/Experimental/DataLinkOAuth/Source/DataLinkOAuthEditor/Private/DataLinkOAuthSettingsFactory.h
--- a/Engine/Plugins/Experimental/DataLinkOAuth/Source/DataLinkOAuthEditor/Private/DataLinkOAuthSettingsFactory.h
+++ b/Engine/Plugins/Experimental/DataLinkOAuth/Source/DataLinkOAuthEditor/Private/DataLinkOAuthSettingsFactory.h
@@ -16,6 +16,8 @@ public:
 	UDataLinkOAuthSettingsFactory();
 
 	//~ Begin UFactory
+	virtual FText GetDisplayName() const override;
+	virtual FString GetDefaultNewAssetName() const override;
 	virtual uint32 GetMenuCategories() const override;
 	virtual bool ConfigureProperties() override;
 	virtual UObject* FactoryCreateNew(UClass* InClass, UObject* InParent, FName InName, EObjectFlags InFlags, UObject* InContext, FFeedbackContext* InWarn) override;
diff --git a/Engine/Plugins/Experimental/SceneState/Source/SceneStateBlueprint/Public/SceneStateBlueprint.h b/Engine/Plugins/Experimental/SceneState/Source/SceneStateBlueprint/Public/SceneStateBlueprint.h
--- a/Engine/Plugins/Experimental/SceneState/Source/SceneStateBlueprint/Public/SceneStateBlueprint.h
+++ b/Engine/Plugins/Experimental/SceneState/Source/SceneStateBlueprint/Public/SceneStateBlueprint.h
@@ -11,7 +11,7 @@
 class USceneStateGeneratedClass;
 class USceneStateMachineStateNode;
 
-UCLASS(MinimalAPI)
+UCLASS(MinimalAPI, DisplayName="Motion Design Scene State Blueprint")
 class USceneStateBlueprint : public UBlueprint, public IPropertyBindingBindingCollectionOwner
 {
 	GENERATED_BODY()
diff --git a/Engine/Plugins/Experimental/SceneState/Source/SceneStateBlueprint/Public/SceneStateTaskBlueprint.h b/Engine/Plugins/Experimental/SceneState/Source/SceneStateBlueprint/Public/SceneStateTaskBlueprint.h
--- a/Engine/Plugins/Experimental/SceneState/Source/SceneStateBlueprint/Public/SceneStateTaskBlueprint.h
+++ b/Engine/Plugins/Experimental/SceneState/Source/SceneStateBlueprint/Public/SceneStateTaskBlueprint.h
@@ -5,7 +5,7 @@
 #include "Engine/Blueprint.h"
 #include "SceneStateTaskBlueprint.generated.h"
 
-UCLASS(MinimalAPI)
+UCLASS(MinimalAPI, DisplayName="Motion Design Scene State Task Blueprint")
 class USceneStateTaskBlueprint : public UBlueprint
 {
 	GENERATED_BODY()
diff --git a/Engine/Plugins/Experimental/SceneState/Source/SceneStateBlueprintEditor/Private/SceneStateBlueprintEditorModule.cpp b/Engine/Plugins/Experimental/SceneState/Source/SceneStateBlueprintEditor/Private/SceneStateBlueprintEditorModule.cpp
--- a/Engine/Plugins/Experimental/SceneState/Source/SceneStateBlueprintEditor/Private/SceneStateBlueprintEditorModule.cpp
+++ b/Engine/Plugins/Experimental/SceneState/Source/SceneStateBlueprintEditor/Private/SceneStateBlueprintEditorModule.cpp
@@ -83,7 +83,7 @@ const FBlueprintEditorModule& FBlueprintEditorModule::GetInternal()
 void FBlueprintEditorModule::StartupModule()
 {
 	FBlueprintEditorCommands::Register();
-	RegisterAssetTools();
+
 	RegisterCompiler();
 	RegisterDefaultEvents();
 	RegisterDetailCustomizations();
@@ -121,12 +121,6 @@ void FBlueprintEditorModule::UnregisterContextEditor(const TSharedPtr<IContextEd
 	ContextEditorRegistry.UnregisterContextEditor(InContextEditor);
 }
 
-void FBlueprintEditorModule::RegisterAssetTools()
-{
-	IAssetTools& AssetTools = FAssetToolsModule::GetModule().Get();
-	AssetTools.RegisterAdvancedAssetCategory(TEXT("SceneStateCategory"), LOCTEXT("AssetCategoryName", "Scene State"));
-}
-
 void FBlueprintEditorModule::RegisterCompiler()
 {
 	IKismetCompilerInterface& KismetCompilerModule = FModuleManager::LoadModuleChecked<IKismetCompilerInterface>("KismetCompiler");
diff --git a/Engine/Plugins/Experimental/SceneState/Source/SceneStateBlueprintEditor/Private/SceneStateBlueprintEditorModule.h b/Engine/Plugins/Experimental/SceneState/Source/SceneStateBlueprintEditor/Private/SceneStateBlueprintEditorModule.h
--- a/Engine/Plugins/Experimental/SceneState/Source/SceneStateBlueprintEditor/Private/SceneStateBlueprintEditorModule.h
+++ b/Engine/Plugins/Experimental/SceneState/Source/SceneStateBlueprintEditor/Private/SceneStateBlueprintEditorModule.h
@@ -30,8 +30,6 @@ protected:
 	virtual void UnregisterContextEditor(const TSharedPtr<IContextEditor>& InContextEditor) override;
 	//~ End IBlueprintEditorModule
 
-	void RegisterAssetTools();
-
 	void RegisterCompiler();
 
 	void RegisterDefaultEvents();
diff --git a/Engine/Plugins/Experimental/SceneState/Source/SceneStateBlueprintEditor/Private/SceneStateBlueprintFactory.cpp b/Engine/Plugins/Experimental/SceneState/Source/SceneStateBlueprintEditor/Private/SceneStateBlueprintFactory.cpp
--- a/Engine/Plugins/Experimental/SceneState/Source/SceneStateBlueprintEditor/Private/SceneStateBlueprintFactory.cpp
+++ b/Engine/Plugins/Experimental/SceneState/Source/SceneStateBlueprintEditor/Private/SceneStateBlueprintFactory.cpp
@@ -46,10 +46,21 @@ UEdGraph* USceneStateBlueprintFactory::AddStateMachine(USceneStateBlueprint* InB
 	return NewStateMachineGraph;
 }
 
+FText USceneStateBlueprintFactory::GetDisplayName() const
+{
+	return SupportedClass ? SupportedClass->GetDisplayNameText() : Super::GetDisplayName();
+}
+
+FString USceneStateBlueprintFactory::GetDefaultNewAssetName() const
+{
+	// Short name removing "Motion Design" and "SceneState" prefix for new assets
+	return TEXT("NewBlueprint");
+}
+
 uint32 USceneStateBlueprintFactory::GetMenuCategories() const
 {
 	IAssetTools& AssetTools = FAssetToolsModule::GetModule().Get();
-	return AssetTools.FindAdvancedAssetCategory(TEXT("SceneStateCategory"));
+	return AssetTools.FindAdvancedAssetCategory("MotionDesignCategory");
 }
 
 UObject* USceneStateBlueprintFactory::FactoryCreateNew(UClass* InClass, UObject* InParent, FName InName, EObjectFlags InFlags, UObject* InContext, FFeedbackContext* InWarn)
diff --git a/Engine/Plugins/Experimental/SceneState/Source/SceneStateBlueprintEditor/Private/SceneStateTaskBlueprintFactory.cpp b/Engine/Plugins/Experimental/SceneState/Source/SceneStateBlueprintEditor/Private/SceneStateTaskBlueprintFactory.cpp
--- a/Engine/Plugins/Experimental/SceneState/Source/SceneStateBlueprintEditor/Private/SceneStateTaskBlueprintFactory.cpp
+++ b/Engine/Plugins/Experimental/SceneState/Source/SceneStateBlueprintEditor/Private/SceneStateTaskBlueprintFactory.cpp
@@ -19,10 +19,21 @@ USceneStateTaskBlueprintFactory::USceneStateTaskBlueprintFactory()
 	bEditAfterNew = true;
 }
 
+FText USceneStateTaskBlueprintFactory::GetDisplayName() const
+{
+	return SupportedClass ? SupportedClass->GetDisplayNameText() : Super::GetDisplayName();
+}
+
+FString USceneStateTaskBlueprintFactory::GetDefaultNewAssetName() const
+{
+	// Short name removing "Motion Design" and "SceneState" prefix for new assets
+	return TEXT("NewTaskBlueprint");
+}
+
 uint32 USceneStateTaskBlueprintFactory::GetMenuCategories() const
 {
 	IAssetTools& AssetTools = FAssetToolsModule::GetModule().Get();
-	return AssetTools.FindAdvancedAssetCategory(TEXT("SceneStateCategory"));
+	return AssetTools.FindAdvancedAssetCategory("MotionDesignCategory");
 }
 
 UObject* USceneStateTaskBlueprintFactory::FactoryCreateNew(UClass* InClass, UObject* InParent, FName InName, EObjectFlags InFlags, UObject* InContext, FFeedbackContext* InWarn, FName InCallingContext)
diff --git a/Engine/Plugins/Experimental/SceneState/Source/SceneStateBlueprintEditor/Private/SceneStateTaskBlueprintFactory.h b/Engine/Plugins/Experimental/SceneState/Source/SceneStateBlueprintEditor/Private/SceneStateTaskBlueprintFactory.h
--- a/Engine/Plugins/Experimental/SceneState/Source/SceneStateBlueprintEditor/Private/SceneStateTaskBlueprintFactory.h
+++ b/Engine/Plugins/Experimental/SceneState/Source/SceneStateBlueprintEditor/Private/SceneStateTaskBlueprintFactory.h
@@ -17,6 +17,8 @@ public:
 	USceneStateTaskBlueprintFactory();
 
 	//~ Begin UFactory
+	virtual FText GetDisplayName() const override;
+	virtual FString GetDefaultNewAssetName() const override;
 	virtual uint32 GetMenuCategories() const override;
 	virtual UObject* FactoryCreateNew(UClass* InClass, UObject* InParent, FName InName, EObjectFlags InFlags, UObject* InContext, FFeedbackContext* InWarn, FName InCallingContext) override;
 	//~ End UFactory
diff --git a/Engine/Plugins/Experimental/SceneState/Source/SceneStateBlueprintEditor/Public/SceneStateBlueprintFactory.h b/Engine/Plugins/Experimental/SceneState/Source/SceneStateBlueprintEditor/Public/SceneStateBlueprintFactory.h
--- a/Engine/Plugins/Experimental/SceneState/Source/SceneStateBlueprintEditor/Public/SceneStateBlueprintFactory.h
+++ b/Engine/Plugins/Experimental/SceneState/Source/SceneStateBlueprintEditor/Public/SceneStateBlueprintFactory.h
@@ -21,6 +21,8 @@ public:
 	static UEdGraph* AddStateMachine(USceneStateBlueprint* InBlueprint);
 
 	//~ Begin UFactory
+	virtual FText GetDisplayName() const override;
+	virtual FString GetDefaultNewAssetName() const override;
 	virtual uint32 GetMenuCategories() const override;
 	virtual UObject* FactoryCreateNew(UClass* InClass, UObject* InParent, FName InName, EObjectFlags InFlags, UObject* InContext, FFeedbackContext* InWarn) override;
 	virtual UObject* FactoryCreateNew(UClass* InClass, UObject* InParent, FName InName, EObjectFlags InFlags, UObject* InContext, FFeedbackContext* InWarn, FName InCallingContext) override;
diff --git a/Engine/Plugins/Experimental/SceneState/Source/SceneStateEvent/Public/SceneStateEventSchemaCollection.h b/Engine/Plugins/Experimental/SceneState/Source/SceneStateEvent/Public/SceneStateEventSchemaCollection.h
--- a/Engine/Plugins/Experimental/SceneState/Source/SceneStateEvent/Public/SceneStateEventSchemaCollection.h
+++ b/Engine/Plugins/Experimental/SceneState/Source/SceneStateEvent/Public/SceneStateEventSchemaCollection.h
@@ -8,7 +8,7 @@
 class USceneStateEventSchemaObject;
 
 /** Holds a collection of Event Schemas to group common event types in one place */
-UCLASS(MinimalAPI)
+UCLASS(MinimalAPI, DisplayName="Motion Design Scene State Event Schema Collection")
 class USceneStateEventSchemaCollection : public UObject
 {
 	GENERATED_BODY()
diff --git a/Engine/Plugins/Experimental/SceneState/Source/SceneStateEventEditor/Private/SceneStateEventSchemaCollectionFactory.cpp b/Engine/Plugins/Experimental/SceneState/Source/SceneStateEventEditor/Private/SceneStateEventSchemaCollectionFactory.cpp
--- a/Engine/Plugins/Experimental/SceneState/Source/SceneStateEventEditor/Private/SceneStateEventSchemaCollectionFactory.cpp
+++ b/Engine/Plugins/Experimental/SceneState/Source/SceneStateEventEditor/Private/SceneStateEventSchemaCollectionFactory.cpp
@@ -13,10 +13,21 @@ USceneStateEventSchemaCollectionFactory::USceneStateEventSchemaCollectionFactory
 	bEditAfterNew = true;
 }
 
+FText USceneStateEventSchemaCollectionFactory::GetDisplayName() const
+{
+	return SupportedClass ? SupportedClass->GetDisplayNameText() : Super::GetDisplayName();
+}
+
+FString USceneStateEventSchemaCollectionFactory::GetDefaultNewAssetName() const
+{
+	// Short name removing "Motion Design" and "Scene State" prefix for new assets
+	return TEXT("NewEventSchemaCollection");
+}
+
 uint32 USceneStateEventSchemaCollectionFactory::GetMenuCategories() const
 {
 	IAssetTools& AssetTools = FAssetToolsModule::GetModule().Get();
-	return AssetTools.FindAdvancedAssetCategory(TEXT("SceneStateCategory"));
+	return AssetTools.FindAdvancedAssetCategory("MotionDesignCategory");
 }
 
 UObject* USceneStateEventSchemaCollectionFactory::FactoryCreateNew(UClass* InClass, UObject* InParent, FName InName, EObjectFlags InFlags, UObject* InContext, FFeedbackContext* InWarn)
diff --git a/Engine/Plugins/Experimental/SceneState/Source/SceneStateEventEditor/Private/SceneStateEventSchemaCollectionFactory.h b/Engine/Plugins/Experimental/SceneState/Source/SceneStateEventEditor/Private/SceneStateEventSchemaCollectionFactory.h
--- a/Engine/Plugins/Experimental/SceneState/Source/SceneStateEventEditor/Private/SceneStateEventSchemaCollectionFactory.h
+++ b/Engine/Plugins/Experimental/SceneState/Source/SceneStateEventEditor/Private/SceneStateEventSchemaCollectionFactory.h
@@ -13,6 +13,8 @@ class USceneStateEventSchemaCollectionFactory : public UFactory
 	USceneStateEventSchemaCollectionFactory();
 
 	//~ Begin UFactory
+	virtual FText GetDisplayName() const override;
+	virtual FString GetDefaultNewAssetName() const override;
 	virtual uint32 GetMenuCategories() const override;
 	virtual UObject* FactoryCreateNew(UClass* InClass, UObject* InParent, FName InName, EObjectFlags InFlags, UObject* InContext, FFeedbackContext* InWarn) override;
 	//~ End UFactory
