Motion Design: fix issue with scene state event broadcast applying to all event streams without a scope (e.g. world). In Rundown, multiple worlds are at play and the events being broadcast there should be self contained.

diff --git a/Engine/Plugins/Experimental/AvalancheSceneState/Source/AvalancheSceneState/Private/RemoteControl/AvaSceneStateRCEventBehaviorNode.cpp b/Engine/Plugins/Experimental/AvalancheSceneState/Source/AvalancheSceneState/Private/RemoteControl/AvaSceneStateRCEventBehaviorNode.cpp
--- a/Engine/Plugins/Experimental/AvalancheSceneState/Source/AvalancheSceneState/Private/RemoteControl/AvaSceneStateRCEventBehaviorNode.cpp
+++ b/Engine/Plugins/Experimental/AvalancheSceneState/Source/AvalancheSceneState/Private/RemoteControl/AvaSceneStateRCEventBehaviorNode.cpp
@@ -1,6 +1,8 @@
 // Copyright Epic Games, Inc. All Rights Reserved.
 
 #include "RemoteControl/AvaSceneStateRCEventBehaviorNode.h"
+#include "AvaSceneStateLog.h"
+#include "Engine/World.h"
 #include "RemoteControl/AvaSceneStateRCEventBehavior.h"
 #include "SceneStateEventUtils.h"
 #include "StructUtils/UserDefinedStruct.h"
@@ -20,12 +22,21 @@ bool UAvaSceneStateRCEventBehaviorNode::Execute(URCBehaviour* InBehavior) const
 
 bool UAvaSceneStateRCEventBehaviorNode::IsSupported(URCBehaviour* InBehavior) const
 {
-	return InBehavior && InBehavior->IsA<UAvaSceneStateRCEventBehavior>();
+	return InBehavior
+		&& InBehavior->IsA<UAvaSceneStateRCEventBehavior>()
+		&& UE::SceneState::GetContextWorld(InBehavior);
 }
 
 void UAvaSceneStateRCEventBehaviorNode::OnPassed(URCBehaviour* InBehavior) const
 {
-	UE::SceneState::BroadcastEvent(Event);
+	if (const UWorld* World = UE::SceneState::GetContextWorld(InBehavior))
+	{
+		UE::SceneState::BroadcastEvent(World, Event);
+	}
+	else
+	{
+		UE_LOG(LogAvaSceneState, Error, TEXT("BroadcastEvent failed. Behavior node %s could not find a valid world."), *GetFullName())
+	}
 }
 
 UClass* UAvaSceneStateRCEventBehaviorNode::GetBehaviourClass() const
diff --git a/Engine/Plugins/Experimental/SceneState/Source/SceneStateEvent/Private/SceneStateEventLibrary.cpp b/Engine/Plugins/Experimental/SceneState/Source/SceneStateEvent/Private/SceneStateEventLibrary.cpp
--- a/Engine/Plugins/Experimental/SceneState/Source/SceneStateEvent/Private/SceneStateEventLibrary.cpp
+++ b/Engine/Plugins/Experimental/SceneState/Source/SceneStateEvent/Private/SceneStateEventLibrary.cpp
@@ -2,8 +2,10 @@
 
 #include "SceneStateEventLibrary.h"
 #include "Blueprint/BlueprintExceptionInfo.h"
+#include "Engine/World.h"
 #include "ISceneStateEventHandlerProvider.h"
 #include "SceneStateEvent.h"
+#include "SceneStateEventLog.h"
 #include "SceneStateEventSchema.h"
 #include "SceneStateEventSchemaHandle.h"
 #include "SceneStateEventStream.h"
@@ -17,9 +19,15 @@ bool USceneStateEventLibrary::PushEvent(USceneStateEventStream* InEventStream, F
 	return UE::SceneState::PushEvent(InEventStream, InEventSchemaHandle, MoveTemp(InEventData));
 }
 
-bool USceneStateEventLibrary::BroadcastEvent(FSceneStateEventSchemaHandle InEventSchemaHandle, FInstancedStruct InEventData)
+bool USceneStateEventLibrary::BroadcastEvent(UObject* WorldContextObject, FSceneStateEventSchemaHandle InEventSchemaHandle, FInstancedStruct InEventData)
 {
-	return UE::SceneState::BroadcastEvent(InEventSchemaHandle, MoveTemp(InEventData));
+	if (const UWorld* World = UE::SceneState::GetContextWorld(WorldContextObject))
+	{
+		return UE::SceneState::BroadcastEvent(World, InEventSchemaHandle, MoveTemp(InEventData));
+	}
+
+	UE_LOG(LogSceneStateEvent, Error, TEXT("BroadcastEvent failed. Could not find a valid world from context object '%s'"), *GetFullNameSafe(WorldContextObject));
+	return false;
 }
 
 bool USceneStateEventLibrary::EventDataToStruct(const FInstancedStruct& InEventData, int32& OutStructValue)
diff --git a/Engine/Plugins/Experimental/SceneState/Source/SceneStateEvent/Private/SceneStateEventSubsystem.cpp b/Engine/Plugins/Experimental/SceneState/Source/SceneStateEvent/Private/SceneStateEventSubsystem.cpp
--- a/Engine/Plugins/Experimental/SceneState/Source/SceneStateEvent/Private/SceneStateEventSubsystem.cpp
+++ b/Engine/Plugins/Experimental/SceneState/Source/SceneStateEvent/Private/SceneStateEventSubsystem.cpp
@@ -19,13 +19,13 @@ void USceneStateEventSubsystem::UnregisterEventStream(USceneStateEventStream* In
 	EventStreams.Remove(InEventStream);
 }
 
-void USceneStateEventSubsystem::ForEachEventStream(TFunctionRef<void(USceneStateEventStream&)> InCallable)
+void USceneStateEventSubsystem::ForEachEventStream(TFunctionRef<void(USceneStateEventStream*)> InCallable)
 {
 	for (USceneStateEventStream* EventStream : EventStreams)
 	{
 		if (EventStream)
 		{
-			InCallable(*EventStream);
+			InCallable(EventStream);
 		}
 	}
 }
diff --git a/Engine/Plugins/Experimental/SceneState/Source/SceneStateEvent/Private/SceneStateEventSubsystem.h b/Engine/Plugins/Experimental/SceneState/Source/SceneStateEvent/Private/SceneStateEventSubsystem.h
--- a/Engine/Plugins/Experimental/SceneState/Source/SceneStateEvent/Private/SceneStateEventSubsystem.h
+++ b/Engine/Plugins/Experimental/SceneState/Source/SceneStateEvent/Private/SceneStateEventSubsystem.h
@@ -17,11 +17,14 @@ class USceneStateEventSubsystem : public UEngineSubsystem
 public:
 	static USceneStateEventSubsystem* Get();
 
+	/** Adds the given event stream to the list of actively registered event streams */
 	void RegisterEventStream(USceneStateEventStream* InEventStream);
 
+	/** Removes the given event stream from the list of actively registered event streams */
 	void UnregisterEventStream(USceneStateEventStream* InEventStream);
 
-	void ForEachEventStream(TFunctionRef<void(USceneStateEventStream&)> InCallable);
+	/** Iterates all valid event streams that are actively registered */
+	void ForEachEventStream(TFunctionRef<void(USceneStateEventStream*)> InCallable);
 
 private:
 	UPROPERTY()
diff --git a/Engine/Plugins/Experimental/SceneState/Source/SceneStateEvent/Private/SceneStateEventUtils.cpp b/Engine/Plugins/Experimental/SceneState/Source/SceneStateEvent/Private/SceneStateEventUtils.cpp
--- a/Engine/Plugins/Experimental/SceneState/Source/SceneStateEvent/Private/SceneStateEventUtils.cpp
+++ b/Engine/Plugins/Experimental/SceneState/Source/SceneStateEvent/Private/SceneStateEventUtils.cpp
@@ -1,6 +1,8 @@
 // Copyright Epic Games, Inc. All Rights Reserved.
 
 #include "SceneStateEventUtils.h"
+#include "Engine/World.h"
+#include "SceneStateEventLog.h"
 #include "SceneStateEventSchema.h"
 #include "SceneStateEventSchemaHandle.h"
 #include "SceneStateEventStream.h"
@@ -11,6 +13,21 @@
 namespace UE::SceneState
 {
 
+namespace Private
+{
+
+bool IsEventStreamInContext(const USceneStateEventStream* InEventStream, const UObject* InEventContext)
+{
+	// Special case for Worlds: instead of checking outer chain, check that the world the event stream returns is the same
+	if (const UWorld* ContextWorld = Cast<UWorld>(InEventContext))
+	{
+		return GetContextWorld(InEventStream) == ContextWorld;
+	}
+	return InEventStream->IsIn(InEventContext);
+}
+
+} // Private
+
 bool PushEvent(USceneStateEventStream* InEventStream, const FSceneStateEventSchemaHandle& InEventSchemaHandle, FInstancedStruct&& InEventData)
 {
 	if (!InEventStream)
@@ -33,34 +50,52 @@ bool PushEvent(USceneStateEventStream* InEventStream, const FSceneStateEventTemp
 	return PushEvent(InEventStream, InEventTemplate.GetEventSchemaHandle(), FInstancedStruct(InEventTemplate.GetEventData()));
 }
 
-bool BroadcastEvent(const FSceneStateEventSchemaHandle& InEventSchemaHandle, FInstancedStruct&& InEventData)
+const UWorld* GetContextWorld(const UObject* InObject)
 {
+	const UWorld* World = InObject ? InObject->GetWorld() : nullptr;
+	return World ? World : GWorld;
+}
+
+bool BroadcastEvent(const UObject* InEventContext, const FSceneStateEventSchemaHandle& InEventSchemaHandle, FInstancedStruct&& InEventData)
+{
+	if (!InEventContext)
+	{
+		UE_LOG(LogSceneStateEvent, Error, TEXT("Broadcast event failed. Event context is invalid."));
+		return false;
+	}
+
 	USceneStateEventSubsystem* const EventSubsystem = USceneStateEventSubsystem::Get();
 	if (!EventSubsystem)
 	{
+		UE_LOG(LogSceneStateEvent, Error, TEXT("Broadcast event failed. Event subsystem not found."));
 		return false;
 	}
 
 	const USceneStateEventSchemaObject* EventSchema = InEventSchemaHandle.GetEventSchema();
 	if (!EventSchema)
 	{
+		UE_LOG(LogSceneStateEvent, Error, TEXT("Broadcast event failed. Event schema not found."));
 		return false;
 	}
 
 	const FSharedStruct Event = EventSchema->CreateEvent(MoveTemp(InEventData));
 
 	EventSubsystem->ForEachEventStream(
-		[&Event](USceneStateEventStream& InEventStream)
+		[&Event, InEventContext](USceneStateEventStream* InEventStream)
 		{
-			InEventStream.PushEvent(Event);
+			// Only consider event streams that are within the context
+			if (Private::IsEventStreamInContext(InEventStream, InEventContext))
+			{
+				InEventStream->PushEvent(Event);
+			}
 		});
 
 	return true;
 }
 
-bool BroadcastEvent(const FSceneStateEventTemplate& InEventTemplate)
+bool BroadcastEvent(const UObject* InEventContext, const FSceneStateEventTemplate& InEventTemplate)
 {
-	return BroadcastEvent(InEventTemplate.GetEventSchemaHandle(), FInstancedStruct(InEventTemplate.GetEventData()));
+	return BroadcastEvent(InEventContext, InEventTemplate.GetEventSchemaHandle(), FInstancedStruct(InEventTemplate.GetEventData()));
 }
 
 } // UE::SceneState
diff --git a/Engine/Plugins/Experimental/SceneState/Source/SceneStateEvent/Public/SceneStateEventLibrary.h b/Engine/Plugins/Experimental/SceneState/Source/SceneStateEvent/Public/SceneStateEventLibrary.h
--- a/Engine/Plugins/Experimental/SceneState/Source/SceneStateEvent/Public/SceneStateEventLibrary.h
+++ b/Engine/Plugins/Experimental/SceneState/Source/SceneStateEvent/Public/SceneStateEventLibrary.h
@@ -15,18 +15,40 @@ class USceneStateEventLibrary : public UBlueprintFunctionLibrary
 	GENERATED_BODY()
 
 public:
+	/**
+	 * Pushes an event to the given event stream
+	 * @param InEventStream the stream to push the event to
+	 * @param InEventSchemaHandle handle to the event schema used to create the event
+	 * @param InEventData the event payload. Must match the event schema struct
+	 * @return true if the event was pushed successfully, false otherwise
+	 */
 	UFUNCTION(BlueprintCallable, Category="Scene State|Event", meta=(BlueprintInternalUseOnly="true"))
 	static bool PushEvent(USceneStateEventStream* InEventStream, FSceneStateEventSchemaHandle InEventSchemaHandle, FInstancedStruct InEventData = FInstancedStruct());
 
-	UFUNCTION(BlueprintCallable, Category="Scene State|Event", meta=(BlueprintInternalUseOnly="true"))
-	static bool BroadcastEvent(FSceneStateEventSchemaHandle InEventSchemaHandle, FInstancedStruct InEventData = FInstancedStruct());
-
+	/**
+	 * Broadcasts an event to the every registered event stream within a world.
+	 * @param WorldContextObject context object to get the world
+	 * @param InEventSchemaHandle handle to the event schema used to create the event
+	 * @param InEventData the event payload. Must match the event schema struct
+	 * @return true if the event was broadcasted successfully, false otherwise
+	 */
+	UFUNCTION(BlueprintCallable, Category="Scene State|Event", meta=(BlueprintInternalUseOnly="true", WorldContext="WorldContextObject"))
+	static bool BroadcastEvent(UObject* WorldContextObject, FSceneStateEventSchemaHandle InEventSchemaHandle, FInstancedStruct InEventData = FInstancedStruct());
+
+	/**
+	 * Converts the given event data instanced struct to the wild card if it matches struct types
+	 * @param InEventData the event data instanced struct to convert
+	 * @param OutStructValue the returned value (wild card) if it matches in struct with the event data
+	 * @return true if the conversion succeeeded
+	 */
 	UFUNCTION(BlueprintPure, CustomThunk, Category="Scene State|Event", meta=(CustomStructureParam="OutStructValue", BlueprintInternalUseOnly="true"))
 	static bool EventDataToStruct(UPARAM(Ref) const FInstancedStruct& InEventData, int32& OutStructValue);
 
+	/** Finds a captured event within an Event Stream with option to look into any pushed event in that stream (not only captured) */
 	UFUNCTION(BlueprintPure, Category="Scene State|Event", meta=(BlueprintInternalUseOnly="true"))
 	static bool FindEvent(UObject* InContextObject, USceneStateEventStream* InEventStream, FSceneStateEventSchemaHandle InEventSchemaHandle, const FString& InEventHandlerId, bool bInCapturedEventsOnly, FInstancedStruct& OutEventData);
 
+	/** Returns whether an event is present as a captured event (or optionally just pushed) within an Event Stream */
 	UFUNCTION(BlueprintPure, Category="Scene State|Event", meta=(BlueprintInternalUseOnly="true"))
 	static bool HasEvent(UObject* InContextObject, USceneStateEventStream* InEventStream, FSceneStateEventSchemaHandle InEventSchemaHandle, const FString& InEventHandlerId, bool bInCapturedEventsOnly);
 
diff --git a/Engine/Plugins/Experimental/SceneState/Source/SceneStateEvent/Public/SceneStateEventUtils.h b/Engine/Plugins/Experimental/SceneState/Source/SceneStateEvent/Public/SceneStateEventUtils.h
--- a/Engine/Plugins/Experimental/SceneState/Source/SceneStateEvent/Public/SceneStateEventUtils.h
+++ b/Engine/Plugins/Experimental/SceneState/Source/SceneStateEvent/Public/SceneStateEventUtils.h
@@ -4,7 +4,9 @@
 
 #include "HAL/Platform.h"
 
+class UObject;
 class USceneStateEventStream;
+class UWorld;
 struct FInstancedStruct;
 struct FSceneStateEventSchemaHandle;
 struct FSceneStateEventTemplate;
@@ -29,18 +31,27 @@ namespace UE::SceneState
 	SCENESTATEEVENT_API bool PushEvent(USceneStateEventStream* InEventStream, const FSceneStateEventTemplate& InEventTemplate);
 
 	/**
-	 * Broadcasts an event to the every registered event stream
+	 * Retrieves a context world for a given object.
+	 * @param InObject the context object to look a context world for
+	 * @return the context world or GWorld as fallback if no explicit world was found.
+	 */
+	SCENESTATEEVENT_API const UWorld* GetContextWorld(const UObject* InObject);
+
+	/**
+	 * Broadcasts an event to the every registered event stream within a given scope.
+	 * @param InEventContext required scope to which the event will be broadcasted to. Only Event Streams within this scope will receive the broadcast event.
 	 * @param InEventSchemaHandle handle to the event schema used to create the event
 	 * @param InEventData the event payload. Must match the event schema struct
 	 * @return true if the event was broadcasted successfully, false otherwise
 	 */
-	SCENESTATEEVENT_API bool BroadcastEvent(const FSceneStateEventSchemaHandle& InEventSchemaHandle, FInstancedStruct&& InEventData);
+	SCENESTATEEVENT_API bool BroadcastEvent(const UObject* InEventContext, const FSceneStateEventSchemaHandle& InEventSchemaHandle, FInstancedStruct&& InEventData);
 
 	/**
-	 * Broadcasts an event to the every registered event stream
+	 * Broadcasts an event to the every registered event stream within a given scope.
+	 * @param InEventContext required scope to which the event will be broadcasted to. Only Event Streams within this scope will receive the broadcast event.
 	 * @param InEventTemplate the template used to create the event
 	 * @return true if the event was broadcasted successfully, false otherwise
 	 */
-	SCENESTATEEVENT_API bool BroadcastEvent(const FSceneStateEventTemplate& InEventTemplate);
+	SCENESTATEEVENT_API bool BroadcastEvent(const UObject* InEventContext, const FSceneStateEventTemplate& InEventTemplate);
 
 } // UE::SceneState
diff --git a/Engine/Plugins/Experimental/SceneState/Source/SceneStateEventGraph/Private/K2Node_SceneStateBroadcastEvent.cpp b/Engine/Plugins/Experimental/SceneState/Source/SceneStateEventGraph/Private/K2Node_SceneStateBroadcastEvent.cpp
--- a/Engine/Plugins/Experimental/SceneState/Source/SceneStateEventGraph/Private/K2Node_SceneStateBroadcastEvent.cpp
+++ b/Engine/Plugins/Experimental/SceneState/Source/SceneStateEventGraph/Private/K2Node_SceneStateBroadcastEvent.cpp
@@ -28,6 +28,12 @@ void UK2Node_SceneStateBroadcastEvent::AllocateDefaultPins()
 	CreatePin(EGPD_Output
 		, UEdGraphSchema_K2::PC_Exec
 		, UEdGraphSchema_K2::PN_Then);
+
+	// Add World Context object pin
+	CreatePin(EGPD_Input
+		, UEdGraphSchema_K2::PC_Object
+		, UObject::StaticClass()
+		, UK2Node_SceneStateEventBase::PN_WorldContextObject);
 }
 
 FText UK2Node_SceneStateBroadcastEvent::GetNodeTitle(ENodeTitleType::Type InTitleType) const
@@ -66,6 +72,13 @@ void UK2Node_SceneStateBroadcastEvent::ExpandNode(FKismetCompilerContext& InComp
 		EventSchemaHandlePin->DefaultValue = GetSchemaHandleStringValue();
 	}
 
+	// Move the World context object pin
+	if (!ConnectPinsToIntermediate(InCompilerContext, BroadcastEventNode, UK2Node_SceneStateEventBase::PN_WorldContextObject, TEXT("WorldContextObject")))
+	{
+		InCompilerContext.MessageLog.Error(*LOCTEXT("EventScopeConnectError", "ICE: Error connecting World Context Object Pin @@").ToString()
+			, this);
+	}
+	
 	FNodeExpansionContext Context = {
 		.CompilerContext = InCompilerContext,
 		.SourceGraph = InSourceGraph,
diff --git a/Engine/Plugins/Experimental/SceneState/Source/SceneStateEventGraph/Private/K2Node_SceneStateEventBase.cpp b/Engine/Plugins/Experimental/SceneState/Source/SceneStateEventGraph/Private/K2Node_SceneStateEventBase.cpp
--- a/Engine/Plugins/Experimental/SceneState/Source/SceneStateEventGraph/Private/K2Node_SceneStateEventBase.cpp
+++ b/Engine/Plugins/Experimental/SceneState/Source/SceneStateEventGraph/Private/K2Node_SceneStateEventBase.cpp
@@ -15,6 +15,7 @@
 #define LOCTEXT_NAMESPACE "K2Node_SceneStateEventBase"
 
 const FLazyName UK2Node_SceneStateEventBase::PN_EventStream = TEXT("EventStream");
+const FLazyName UK2Node_SceneStateEventBase::PN_WorldContextObject = TEXT("WorldContextObject");
 
 FText UK2Node_SceneStateEventBase::GetSchemaDisplayNameText() const
 {
@@ -82,7 +83,8 @@ bool UK2Node_SceneStateEventBase::IsEventDataPin(const UEdGraphPin* InPin) const
 	return InPin->PinName != UEdGraphSchema_K2::PN_Execute &&
 		InPin->PinName != UEdGraphSchema_K2::PN_Then &&
 		InPin->PinName != UEdGraphSchema_K2::PN_ReturnValue &&
-		InPin->PinName != UK2Node_SceneStateEventBase::PN_EventStream;
+		InPin->PinName != UK2Node_SceneStateEventBase::PN_EventStream &&
+		InPin->PinName != UK2Node_SceneStateEventBase::PN_WorldContextObject;
 }
 
 void UK2Node_SceneStateEventBase::CreateEventDataPins(TConstArrayView<UEdGraphPin*> InPinsToSearch)
diff --git a/Engine/Plugins/Experimental/SceneState/Source/SceneStateEventGraph/Private/K2Node_SceneStateEventBase.h b/Engine/Plugins/Experimental/SceneState/Source/SceneStateEventGraph/Private/K2Node_SceneStateEventBase.h
--- a/Engine/Plugins/Experimental/SceneState/Source/SceneStateEventGraph/Private/K2Node_SceneStateEventBase.h
+++ b/Engine/Plugins/Experimental/SceneState/Source/SceneStateEventGraph/Private/K2Node_SceneStateEventBase.h
@@ -12,7 +12,9 @@ class UK2Node_SceneStateEventBase : public UK2Node
 	GENERATED_BODY()
 
 protected:
+	/** Common Pin Names */
 	static const FLazyName PN_EventStream;
+	static const FLazyName PN_WorldContextObject;
 
 	FText GetSchemaDisplayNameText() const;
 
