Motion Design: fixed issue where dragging a function into the scene state machine graph would create an invalid K2 function node and interacting with it would result in a crash. Ideally, the SMyBlueprint drag functionality should be fixed to allow implementations to disallow functions from being dragged in. As a current temporary measure, even if the entire drag operation takes place, it will be blocked by this 'Add Node' change.

diff --git a/Engine/Plugins/Experimental/SceneState/Source/SceneStateMachineGraph/Private/SceneStateMachineGraph.cpp b/Engine/Plugins/Experimental/SceneState/Source/SceneStateMachineGraph/Private/SceneStateMachineGraph.cpp
--- a/Engine/Plugins/Experimental/SceneState/Source/SceneStateMachineGraph/Private/SceneStateMachineGraph.cpp
+++ b/Engine/Plugins/Experimental/SceneState/Source/SceneStateMachineGraph/Private/SceneStateMachineGraph.cpp
@@ -46,6 +46,21 @@ USceneStateMachineEntryNode* USceneStateMachineGraph::GetEntryNode() const
 	return LastEntryNode;
 }
 
+void USceneStateMachineGraph::AddNode(UEdGraphNode* InNodeToAdd, bool bInUserAction, bool bInSelectNewNode)
+{
+	if (!InNodeToAdd)
+	{
+		return;
+	}
+
+	// Workaround for when 'CanCreateUnderSpecifiedSchema' is not called on situations like in SMyBlueprint::OnActionDragged for functions
+	const UEdGraphSchema* GraphSchema = GetSchema();
+	if (GraphSchema && InNodeToAdd->CanCreateUnderSpecifiedSchema(GraphSchema))
+	{
+		Super::AddNode(InNodeToAdd, bInUserAction, bInSelectNewNode);
+	}
+}
+
 void USceneStateMachineGraph::PostInitProperties()
 {
 	Super::PostInitProperties();
diff --git a/Engine/Plugins/Experimental/SceneState/Source/SceneStateMachineGraph/Public/SceneStateMachineGraph.h b/Engine/Plugins/Experimental/SceneState/Source/SceneStateMachineGraph/Public/SceneStateMachineGraph.h
--- a/Engine/Plugins/Experimental/SceneState/Source/SceneStateMachineGraph/Public/SceneStateMachineGraph.h
+++ b/Engine/Plugins/Experimental/SceneState/Source/SceneStateMachineGraph/Public/SceneStateMachineGraph.h
@@ -25,6 +25,10 @@ public:
 
 	SCENESTATEMACHINEGRAPH_API USceneStateMachineEntryNode* GetEntryNode() const;
 
+	//~ Begin UEdGraph
+	SCENESTATEMACHINEGRAPH_API virtual void AddNode(UEdGraphNode* InNodeToAdd, bool bInUserAction, bool bInSelectNewNode) override;
+	//~ End UEdGraph
+
 	//~ Begin UObject
 	SCENESTATEMACHINEGRAPH_API virtual void PostInitProperties() override;
 	SCENESTATEMACHINEGRAPH_API virtual void PostLoad() override;
