Motion Design: fixed issues with renaming state machine parameters breaking existing bindings

diff --git a/Engine/Plugins/Experimental/SceneState/Source/SceneStateBlueprint/Private/SceneStateBlueprint.cpp b/Engine/Plugins/Experimental/SceneState/Source/SceneStateBlueprint/Private/SceneStateBlueprint.cpp
--- a/Engine/Plugins/Experimental/SceneState/Source/SceneStateBlueprint/Private/SceneStateBlueprint.cpp
+++ b/Engine/Plugins/Experimental/SceneState/Source/SceneStateBlueprint/Private/SceneStateBlueprint.cpp
@@ -10,6 +10,7 @@
 #include "Nodes/SceneStateMachineTaskNode.h"
 #include "PropertyBindingDataView.h"
 #include "SceneStateBindingDesc.h"
+#include "SceneStateBindingUtils.h"
 #include "SceneStateBlueprintBindingUtils.h"
 #include "SceneStateBlueprintDelegates.h"
 #include "SceneStateBlueprintLog.h"
@@ -24,6 +25,13 @@ USceneStateBlueprint::USceneStateBlueprint(const FObjectInitializer& InObjectIni
 
 	OnRenameVariableReferencesHandle = FBlueprintEditorUtils::OnRenameVariableReferencesEvent.AddUObject(this
 		, &USceneStateBlueprint::OnRenameVariableReferences);
+
+	OnGraphParametersChangedHandle = USceneStateMachineGraph::OnParametersChanged().AddUObject(this
+		, &USceneStateBlueprint::OnGraphParametersChanged);
+
+#if WITH_EDITOR
+	BindingCollection.SetBindingsOwner(this);
+#endif // WITH_EDITOR
 }
 
 UBlueprintExtension* USceneStateBlueprint::FindExtension(TSubclassOf<UBlueprintExtension> InClass) const
@@ -182,6 +190,9 @@ void USceneStateBlueprint::BeginDestroy()
 
 	FBlueprintEditorUtils::OnRenameVariableReferencesEvent.Remove(OnRenameVariableReferencesHandle);
 	OnRenameVariableReferencesHandle.Reset();
+
+	USceneStateMachineGraph::OnParametersChanged().Remove(OnGraphParametersChangedHandle);
+	OnGraphParametersChangedHandle.Reset();
 }
 
 void USceneStateBlueprint::OnRenameVariableReferences(UBlueprint* InBlueprint, UClass* InVariableClass, const FName& InOldVariableName, const FName& InNewVariableName)
@@ -230,3 +241,53 @@ void USceneStateBlueprint::RenameVariableReferenceInPath(FPropertyBindingPath& I
 			, *InPath.ToString());
 	}
 }
+
+void USceneStateBlueprint::OnGraphParametersChanged(USceneStateMachineGraph* InGraph)
+{
+	BindingCollection.ForEachMutableBinding(
+		[This=this, InGraph](FPropertyBindingBinding& InBinding)
+		{
+			This->UpdateGraphParametersBindings(InBinding.GetMutableSourcePath(), InGraph);
+			This->UpdateGraphParametersBindings(InBinding.GetMutableTargetPath(), InGraph);
+		});
+}
+
+void USceneStateBlueprint::UpdateGraphParametersBindings(FPropertyBindingPath& InPath, USceneStateMachineGraph* InGraph)
+{
+	// Only consider fixing paths that are set to the graph parameters
+	if (InPath.GetStructID() != InGraph->ParametersId)
+	{
+		return;
+	}
+
+	TArrayView<FPropertyBindingPathSegment> Segments = InPath.GetMutableSegments();
+	if (Segments.IsEmpty())
+	{
+		return;
+	}
+
+	// Only need to consider the first segment of the path (i.e. the segment that might've been renamed)
+	FPropertyBindingPathSegment& Segment = Segments[0];
+
+	if (const FPropertyBagPropertyDesc* Property = InGraph->Parameters.FindPropertyDescByID(Segment.GetPropertyGuid()))
+	{
+		const FName SegmentName = Segment.GetName();
+
+		if (SegmentName != Property->Name)
+		{
+			if (GUndo)
+			{
+				Modify();
+			}
+
+			const FString OldPath = InPath.ToString();
+			Segment.SetName(Property->Name);
+
+			UE_LOG(LogSceneStateBlueprint, Log, TEXT("Renamed parameter variable binding segment '%s' to '%s'. (OldPath: %s ---> New Path: %s)")
+				, *SegmentName.ToString()
+				, *Property->Name.ToString()
+				, *OldPath
+				, *InPath.ToString());
+		}
+	}
+}
diff --git a/Engine/Plugins/Experimental/SceneState/Source/SceneStateBlueprint/Public/SceneStateBlueprint.h b/Engine/Plugins/Experimental/SceneState/Source/SceneStateBlueprint/Public/SceneStateBlueprint.h
--- a/Engine/Plugins/Experimental/SceneState/Source/SceneStateBlueprint/Public/SceneStateBlueprint.h
+++ b/Engine/Plugins/Experimental/SceneState/Source/SceneStateBlueprint/Public/SceneStateBlueprint.h
@@ -9,6 +9,7 @@
 #include "SceneStateBlueprint.generated.h"
 
 class USceneStateGeneratedClass;
+class USceneStateMachineGraph;
 class USceneStateMachineStateNode;
 
 UCLASS(MinimalAPI, DisplayName="Motion Design Scene State Blueprint")
@@ -72,9 +73,19 @@ private:
 	/** Called when a blueprint variable has been renamed. Used to fix the path to the variable being renamed to use the new name */
 	void RenameVariableReferenceInPath(FPropertyBindingPath& InPath, FName InOldVariableName, FName InNewVariableName);
 
+	/** Called when a state machine graph's parameters have changed */
+	void OnGraphParametersChanged(USceneStateMachineGraph* InGraph);
+
+	/** Called when a state machine graph's parameters have changed. Used to fix the path to parameters that have possibly been renamed */
+	void UpdateGraphParametersBindings(FPropertyBindingPath& InPath, USceneStateMachineGraph* InGraph);
+
 	/** Unique id representing this blueprint / generated class as a bindable struct */
 	UPROPERTY()
 	FGuid RootId;
 
+	/** Handle to the delegate when a blueprint variable has been renamed */
 	FDelegateHandle OnRenameVariableReferencesHandle;
+
+	/** Handle to the delegate when a state machine graph parameters have changed */
+	FDelegateHandle OnGraphParametersChangedHandle;
 };
diff --git a/Engine/Plugins/Experimental/SceneState/Source/SceneStateBlueprintEditor/Private/DetailsView/Widgets/SSceneStateMachinePicker.cpp b/Engine/Plugins/Experimental/SceneState/Source/SceneStateBlueprintEditor/Private/DetailsView/Widgets/SSceneStateMachinePicker.cpp
--- a/Engine/Plugins/Experimental/SceneState/Source/SceneStateBlueprintEditor/Private/DetailsView/Widgets/SSceneStateMachinePicker.cpp
+++ b/Engine/Plugins/Experimental/SceneState/Source/SceneStateBlueprintEditor/Private/DetailsView/Widgets/SSceneStateMachinePicker.cpp
@@ -87,16 +87,16 @@ SStateMachinePicker::~SStateMachinePicker()
 	OnParametersChangedHandle.Reset();
 }
 
-void SStateMachinePicker::OnParametersChanged(USceneStateMachineGraph& InGraph)
+void SStateMachinePicker::OnParametersChanged(USceneStateMachineGraph* InGraph)
 {
-	if (!SelectedOption.IsValid() || SelectedOption->Id != InGraph.ParametersId)
+	if (!SelectedOption.IsValid() || SelectedOption->Id != InGraph->ParametersId)
 	{
 		return;
 	}
 
 	FScopedTransaction Transaction(LOCTEXT("ParametersChanged", "Parameters Changed"));
 
-	WriteGuidToProperty(StateMachineIdHandle, InGraph.ParametersId);
+	WriteGuidToProperty(StateMachineIdHandle, InGraph->ParametersId);
 	RefreshOptions();
 
 	ParametersHandle->NotifyPreChange();
diff --git a/Engine/Plugins/Experimental/SceneState/Source/SceneStateBlueprintEditor/Private/DetailsView/Widgets/SSceneStateMachinePicker.h b/Engine/Plugins/Experimental/SceneState/Source/SceneStateBlueprintEditor/Private/DetailsView/Widgets/SSceneStateMachinePicker.h
--- a/Engine/Plugins/Experimental/SceneState/Source/SceneStateBlueprintEditor/Private/DetailsView/Widgets/SSceneStateMachinePicker.h
+++ b/Engine/Plugins/Experimental/SceneState/Source/SceneStateBlueprintEditor/Private/DetailsView/Widgets/SSceneStateMachinePicker.h
@@ -36,7 +36,7 @@ public:
 	virtual ~SStateMachinePicker() override;
 
 private:
-	void OnParametersChanged(USceneStateMachineGraph& InGraph);
+	void OnParametersChanged(USceneStateMachineGraph* InGraph);
 
 	FText GetStateMachineName() const;
 
diff --git a/Engine/Plugins/Experimental/SceneState/Source/SceneStateMachineGraph/Private/SceneStateMachineGraph.cpp b/Engine/Plugins/Experimental/SceneState/Source/SceneStateMachineGraph/Private/SceneStateMachineGraph.cpp
--- a/Engine/Plugins/Experimental/SceneState/Source/SceneStateMachineGraph/Private/SceneStateMachineGraph.cpp
+++ b/Engine/Plugins/Experimental/SceneState/Source/SceneStateMachineGraph/Private/SceneStateMachineGraph.cpp
@@ -19,7 +19,7 @@ USceneStateMachineGraph::FOnParametersChanged::RegistrationType& USceneStateMach
 
 void USceneStateMachineGraph::NotifyParametersChanged()
 {
-	OnParametersChangedDelegate.Broadcast(*this);
+	OnParametersChangedDelegate.Broadcast(this);
 }
 
 USceneStateMachineEntryNode* USceneStateMachineGraph::GetEntryNode() const
diff --git a/Engine/Plugins/Experimental/SceneState/Source/SceneStateMachineGraph/Public/SceneStateMachineGraph.h b/Engine/Plugins/Experimental/SceneState/Source/SceneStateMachineGraph/Public/SceneStateMachineGraph.h
--- a/Engine/Plugins/Experimental/SceneState/Source/SceneStateMachineGraph/Public/SceneStateMachineGraph.h
+++ b/Engine/Plugins/Experimental/SceneState/Source/SceneStateMachineGraph/Public/SceneStateMachineGraph.h
@@ -15,7 +15,7 @@ class USceneStateMachineGraph : public UEdGraph
 	GENERATED_BODY()
 
 public:
-	using FOnParametersChanged = TMulticastDelegate<void(USceneStateMachineGraph&)>;
+	using FOnParametersChanged = TMulticastDelegate<void(USceneStateMachineGraph*)>;
 
 	USceneStateMachineGraph();
 
