Scene State: Added property reference support to the binding system.

diff --git a/Engine/Plugins/Experimental/SceneState/Source/SceneState/Private/SceneStatePropertyUtils.cpp b/Engine/Plugins/Experimental/SceneState/Source/SceneState/Private/SceneStatePropertyUtils.cpp
new file mode 100644
--- /dev/null
+++ b/Engine/Plugins/Experimental/SceneState/Source/SceneState/Private/SceneStatePropertyUtils.cpp
@@ -0,0 +1,37 @@
+// Copyright Epic Games, Inc. All Rights Reserved.
+
+#include "SceneStatePropertyUtils.h"
+#include "PropertyBindingDataView.h"
+#include "SceneStateBindingCollection.h"
+#include "SceneStateExecutionContext.h"
+
+namespace UE::SceneState
+{
+
+const FSceneStateBindingResolvedReference* FindResolvedReference(const FSceneStateExecutionContext& InContext, const FSceneStatePropertyReference& InPropertyReference)
+{
+	if (const FSceneStateBindingCollection* BindingCollection = InContext.GetBindingCollection())
+	{
+		return BindingCollection->FindResolvedReference(InPropertyReference);
+	}
+	return nullptr;
+}
+
+uint8* ResolveProperty(const FSceneStateExecutionContext& InContext, const FSceneStateBindingResolvedReference& InResolvedReference)
+{
+	const FSceneStateBindingCollection* BindingCollection = InContext.GetBindingCollection();
+	if (!BindingCollection)
+	{
+		return nullptr;
+	}
+
+	const FPropertyBindingDataView SourceDataView = InContext.FindDataView(InResolvedReference.SourceDataHandle);
+	if (!SourceDataView.IsValid())
+	{
+		return nullptr;
+	}
+
+	return BindingCollection->ResolveProperty(InResolvedReference, SourceDataView);
+}
+
+} // UE::SceneState
diff --git a/Engine/Plugins/Experimental/SceneState/Source/SceneState/Public/SceneStateExecutionContext.h b/Engine/Plugins/Experimental/SceneState/Source/SceneState/Public/SceneStateExecutionContext.h
--- a/Engine/Plugins/Experimental/SceneState/Source/SceneState/Public/SceneStateExecutionContext.h
+++ b/Engine/Plugins/Experimental/SceneState/Source/SceneState/Public/SceneStateExecutionContext.h
@@ -5,6 +5,7 @@
 #include "Containers/Array.h"
 #include "SceneStateInstance.h"
 #include "SceneStateMachineInstance.h"
+#include "SceneStatePropertyUtils.h"
 #include "StructUtils/InstancedStructContainer.h"
 #include "Transition/SceneStateTransitionInstance.h"
 #include "UObject/ObjectPtr.h"
@@ -24,6 +25,7 @@ struct FSceneStateConduit;
 struct FSceneStateEventHandler;
 struct FSceneStateMachine;
 struct FSceneStateMetadata;
+struct FSceneStatePropertyReference;
 struct FSceneStateTask;
 struct FSceneStateTransition;
 
@@ -88,7 +90,14 @@ struct FSceneStateExecutionContext
 	SCENESTATE_API void Reset();
 
 	/** Finds the data view that matches the given data handle */
-	FPropertyBindingDataView FindDataView(const FSceneStateBindingDataHandle& InDataHandle) const;
+	SCENESTATE_API FPropertyBindingDataView FindDataView(const FSceneStateBindingDataHandle& InDataHandle) const;
+
+	/** Returns the pointer to the property value address the given reference is mapped to */
+	template<typename T = void>
+	T* ResolveProperty(const FSceneStatePropertyReference& InPropertyReference) const
+	{
+		return UE::SceneState::ResolveProperty<T>(*this, InPropertyReference);
+	}
 
 	/** Invokes a given callable for each task in the given state */
 	void ForEachTask(const FSceneState& InState, TFunctionRef<UE::SceneState::EIterationResult(const FSceneStateTask&, FStructView)> InCallable) const;
diff --git a/Engine/Plugins/Experimental/SceneState/Source/SceneState/Public/SceneStatePropertyUtils.h b/Engine/Plugins/Experimental/SceneState/Source/SceneState/Public/SceneStatePropertyUtils.h
new file mode 100644
--- /dev/null
+++ b/Engine/Plugins/Experimental/SceneState/Source/SceneState/Public/SceneStatePropertyUtils.h
@@ -0,0 +1,40 @@
+// Copyright Epic Games, Inc. All Rights Reserved.
+
+#pragma once
+
+#include "SceneStateBindingReference.h"
+#include "SceneStatePropertyReferenceUtils.h"
+
+struct FSceneStateExecutionContext;
+struct FSceneStatePropertyReference;
+
+namespace UE::SceneState
+{
+	/**
+	 * Finds the resolved reference for the given property reference
+	 * @param InContext the execution context containing the binding collection
+	 * @param InPropertyReference the property reference to resolve
+	 * @return the resolved reference if found, null otherwise
+	 */
+	SCENESTATE_API const FSceneStateBindingResolvedReference* FindResolvedReference(const FSceneStateExecutionContext& InContext, const FSceneStatePropertyReference& InPropertyReference);
+
+	/**
+	 * Resolves the given reference and execution context to the value address of the property
+	 * @param InContext the execution context containing the binding collection and way to retrieve the data views
+	 * @param InResolvedReference the property reference to resolve access for
+	 * @return value address of the property if successful, null otherwise
+	 */
+	SCENESTATE_API uint8* ResolveProperty(const FSceneStateExecutionContext& InContext, const FSceneStateBindingResolvedReference& InResolvedReference);
+
+	template<class T>
+	T* ResolveProperty(const FSceneStateExecutionContext& InContext, const FSceneStatePropertyReference& InPropertyReference)
+	{
+		const FSceneStateBindingResolvedReference* ResolvedReference = FindResolvedReference(InContext, InPropertyReference);
+		if (ResolvedReference && ResolvedReference->SourceLeafProperty && TDataTypeHelper<T>::IsValid(ResolvedReference->SourceLeafProperty))
+		{
+			return reinterpret_cast<T*>(ResolveProperty(InContext, *ResolvedReference));
+		}
+		return nullptr;
+	}
+
+} // UE::SceneState
diff --git a/Engine/Plugins/Experimental/SceneState/Source/SceneStateBinding/Private/SceneStateBindingCollection.cpp b/Engine/Plugins/Experimental/SceneState/Source/SceneStateBinding/Private/SceneStateBindingCollection.cpp
--- a/Engine/Plugins/Experimental/SceneState/Source/SceneStateBinding/Private/SceneStateBindingCollection.cpp
+++ b/Engine/Plugins/Experimental/SceneState/Source/SceneStateBinding/Private/SceneStateBindingCollection.cpp
@@ -1,6 +1,9 @@
 // Copyright Epic Games, Inc. All Rights Reserved.
 
 #include "SceneStateBindingCollection.h"
+#include "PropertyBindingDataView.h"
+#include "SceneStatePropertyReference.h"
+#include "SceneStatePropertyReferenceUtils.h"
 
 const FSceneStateBindingDesc* FSceneStateBindingCollection::FindBindingDesc(FSceneStateBindingDataHandle InDataHandle) const
 {
@@ -11,6 +14,30 @@ const FSceneStateBindingDesc* FSceneStateBindingCollection::FindBindingDesc(FSce
 		});
 }
 
+const FSceneStateBindingResolvedReference* FSceneStateBindingCollection::FindResolvedReference(const FSceneStatePropertyReference& InPropertyReference) const
+{
+	if (ResolvedReferences.IsValidIndex(InPropertyReference.ReferenceIndex))
+	{
+		return &ResolvedReferences[InPropertyReference.ReferenceIndex];
+	}
+	return nullptr;
+}
+
+uint8* FSceneStateBindingCollection::ResolveProperty(const FSceneStateBindingResolvedReference& InResolvedReference, FPropertyBindingDataView InDataView) const
+{
+	if (!InDataView.IsValid())
+	{
+		return nullptr;
+	}
+
+	if (!ensure(InDataView.GetStruct() == InResolvedReference.SourceStructType))
+	{
+		return nullptr;
+	}
+
+	return Super::GetAddress(InDataView, InResolvedReference.SourceIndirection, InResolvedReference.SourceLeafProperty);
+}
+
 #if WITH_EDITOR
 FPropertyBindingBinding* FSceneStateBindingCollection::AddBindingInternal(const FPropertyBindingPath& InSourcePath, const FPropertyBindingPath& InTargetPath)
 {
@@ -113,3 +140,37 @@ void FSceneStateBindingCollection::VisitSourceStructDescriptorInternal(TFunction
 		}
 	}
 }
+
+bool FSceneStateBindingCollection::OnResolvingPaths()
+{
+	// Base class handled common bindings, here we only need to handle Property references
+	bool bResult = true;
+
+	ResolvedReferences.Reset(References.Num());
+
+	for (const FSceneStateBindingReference& Reference : References)
+	{
+		// Always adds a defaulted reference even in the unexpected occasion binding desc is not found.
+		// This is so that reference indices are kept stable.
+		FSceneStateBindingResolvedReference& ResolvedReference = ResolvedReferences.AddDefaulted_GetRef();
+		ResolvedReference.SourceDataHandle = Reference.SourceDataHandle;
+
+		const FPropertyBindingBindableStructDescriptor* SourceDesc = FindBindingDesc(Reference.SourceDataHandle);
+		if (!ensure(SourceDesc))
+		{
+			continue;
+		}
+
+		ResolvedReference.SourceStructType = SourceDesc->Struct;
+
+		FPropertyBindingPathIndirection SourceLeafIndirection;
+		if (!Super::ResolvePath(SourceDesc->Struct, Reference.SourcePropertyPath, ResolvedReference.SourceIndirection, SourceLeafIndirection))
+		{
+			bResult = false;
+		}
+
+		ResolvedReference.SourceLeafProperty = SourceLeafIndirection.GetProperty();
+	}
+
+	return bResult;
+}
diff --git a/Engine/Plugins/Experimental/SceneState/Source/SceneStateBinding/Private/SceneStatePropertyReferenceUtils.cpp b/Engine/Plugins/Experimental/SceneState/Source/SceneStateBinding/Private/SceneStatePropertyReferenceUtils.cpp
new file mode 100644
--- /dev/null
+++ b/Engine/Plugins/Experimental/SceneState/Source/SceneStateBinding/Private/SceneStatePropertyReferenceUtils.cpp
@@ -0,0 +1,275 @@
+// Copyright Epic Games, Inc. All Rights Reserved.
+
+#include "SceneStatePropertyReferenceUtils.h"
+#include "SceneStatePropertyReference.h"
+
+#if WITH_EDITOR
+#include "EdGraph/EdGraphNode.h"
+#include "EdGraph/EdGraphPin.h"
+#include "EdGraphSchema_K2.h"
+#include "SceneStatePropertyReferenceMetadata.h"
+#endif
+
+namespace UE::SceneState
+{
+
+namespace Private
+{
+
+#if WITH_EDITOR
+/** Supported Property Reference Types */
+static const FLazyName TYPE_Wildcard(TEXT("Wildcard"));
+static const FLazyName TYPE_Bool(TEXT("bool"));
+static const FLazyName TYPE_Byte(TEXT("byte"));
+static const FLazyName TYPE_Int32(TEXT("int32"));
+static const FLazyName TYPE_Int64(TEXT("int64"));
+static const FLazyName TYPE_Float(TEXT("float"));
+static const FLazyName TYPE_Double(TEXT("double"));
+static const FLazyName TYPE_Name(TEXT("Name"));
+static const FLazyName TYPE_String(TEXT("String"));
+static const FLazyName TYPE_Text(TEXT("Text"));
+
+bool ArePropertyReferenceCompatible(const FProperty* InTargetReferenceProperty, const FProperty* InSourceReferenceProperty)
+{
+	check(IsPropertyReference(InSourceReferenceProperty) && IsPropertyReference(InTargetReferenceProperty));
+
+	const FEdGraphPinType SourceRefPin = GetPropertyReferencePinType(InSourceReferenceProperty);
+	const FEdGraphPinType TargetRefPin = GetPropertyReferencePinType(InTargetReferenceProperty);
+
+	return SourceRefPin.PinCategory == TargetRefPin.PinCategory
+		&& SourceRefPin.ContainerType == TargetRefPin.ContainerType 
+		&& SourceRefPin.PinSubCategoryObject == TargetRefPin.PinSubCategoryObject;
+}
+
+bool IsPropertyCompatible(const FString& InTargetType, const FProperty* InTestProperty, bool bInIsTargetRefArray)
+{
+	const FName TargetTypeName = *InTargetType;
+
+	if (TargetTypeName == TYPE_Wildcard)
+	{
+		return true;
+	}
+	if (TargetTypeName == TYPE_Bool)
+	{
+		return InTestProperty->IsA<FBoolProperty>();
+	}
+	if (TargetTypeName == TYPE_Byte)
+	{
+		return InTestProperty->IsA<FByteProperty>();
+	}
+	if (TargetTypeName == TYPE_Int32)
+	{
+		return InTestProperty->IsA<FIntProperty>();
+	}
+	if (TargetTypeName == TYPE_Int64)
+	{
+		return InTestProperty->IsA<FInt64Property>();
+	}
+	if (TargetTypeName == TYPE_Float)
+	{
+		return InTestProperty->IsA<FFloatProperty>();
+	}
+	if (TargetTypeName == TYPE_Double)
+	{
+		return InTestProperty->IsA<FDoubleProperty>();
+	}
+	if (TargetTypeName == TYPE_Name)
+	{
+		return InTestProperty->IsA<FNameProperty>();
+	}
+	if (TargetTypeName == TYPE_String)
+	{
+		return InTestProperty->IsA<FStrProperty>();
+	}
+	if (TargetTypeName == TYPE_Text)
+	{
+		return InTestProperty->IsA<FTextProperty>();
+	}
+
+	UField* TargetRefField = UClass::TryFindTypeSlow<UField>(InTargetType);
+	if (!TargetRefField)
+	{
+		TargetRefField = LoadObject<UField>(nullptr, *InTargetType);
+	}
+
+	if (const FStructProperty* TestStructProperty = CastField<FStructProperty>(InTestProperty))
+	{
+		if (TestStructProperty->Struct && TestStructProperty->Struct->IsChildOf(Cast<UStruct>(TargetRefField)))
+		{
+			return true;
+		}
+	}
+	else if (const FObjectProperty* ObjectProperty = CastField<FObjectProperty>(InTestProperty))
+	{
+		if (ObjectProperty->PropertyClass && ObjectProperty->PropertyClass->IsChildOf(Cast<UClass>(TargetRefField)))
+		{
+			return true;
+		}
+	}
+	else if (const FEnumProperty* EnumProperty = CastField<FEnumProperty>(InTestProperty))
+	{
+		if (EnumProperty->GetEnum() == TargetRefField)
+		{
+			return true;
+		}
+	}
+
+	return false;
+}
+
+/** Retrieves the Ed Graph Pin Type for a given target type (in string) and a desired container type */
+FEdGraphPinType GetTargetPinType(const FString& InTargetType, EPinContainerType InContainerType)
+{
+	if (InTargetType.IsEmpty())
+	{
+		return FEdGraphPinType();
+	}
+
+	const FName TargetTypeName = *InTargetType;
+
+	FEdGraphPinType PinType;
+	PinType.ContainerType = InContainerType;
+
+	if (TargetTypeName == Private::TYPE_Wildcard)
+	{
+		PinType.PinCategory = UEdGraphSchema_K2::PC_Wildcard;
+	}
+	else if (TargetTypeName == Private::TYPE_Bool)
+	{
+		PinType.PinCategory = UEdGraphSchema_K2::PC_Boolean;
+	}
+	else if (TargetTypeName == Private::TYPE_Byte)
+	{
+		PinType.PinCategory = UEdGraphSchema_K2::PC_Byte;
+	}
+	else if (TargetTypeName == Private::TYPE_Int32)
+	{
+		PinType.PinCategory = UEdGraphSchema_K2::PC_Int;
+	}
+	else if (TargetTypeName == Private::TYPE_Int64)
+	{
+		PinType.PinCategory = UEdGraphSchema_K2::PC_Int64;
+	}
+	else if (TargetTypeName == Private::TYPE_Float)
+	{
+		PinType.PinCategory = UEdGraphSchema_K2::PC_Real;
+		PinType.PinSubCategory = UEdGraphSchema_K2::PC_Float;
+	}
+	else if (TargetTypeName == Private::TYPE_Double)
+	{
+		PinType.PinCategory = UEdGraphSchema_K2::PC_Real;
+		PinType.PinSubCategory = UEdGraphSchema_K2::PC_Double;
+	}
+	else if (TargetTypeName == Private::TYPE_Name)
+	{
+		PinType.PinCategory = UEdGraphSchema_K2::PC_Name;
+	}
+	else if (TargetTypeName == Private::TYPE_String)
+	{
+		PinType.PinCategory = UEdGraphSchema_K2::PC_String;
+	}
+	else if (TargetTypeName == Private::TYPE_Text)
+	{
+		PinType.PinCategory = UEdGraphSchema_K2::PC_Text;
+	}
+	else
+	{
+		UField* TargetRefField = UClass::TryFindTypeSlow<UField>(InTargetType);
+		if (!TargetRefField)
+		{
+			TargetRefField = LoadObject<UField>(nullptr, *InTargetType);
+		}
+
+		if (UScriptStruct* Struct = Cast<UScriptStruct>(TargetRefField))
+		{
+			PinType.PinCategory = UEdGraphSchema_K2::PC_Struct;
+			PinType.PinSubCategoryObject = Struct;
+		}
+		else if (UClass* ObjectClass = Cast<UClass>(TargetRefField))
+		{
+			PinType.PinCategory = UEdGraphSchema_K2::PC_Object;
+			PinType.PinSubCategoryObject = ObjectClass;
+		}
+		else if (UEnum* Enum = Cast<UEnum>(TargetRefField))
+		{
+			PinType.PinCategory = UEdGraphSchema_K2::PC_Enum;
+			PinType.PinSubCategoryObject = Enum;
+			if (Enum->GetMaxEnumValue() <= static_cast<int64>(std::numeric_limits<uint8>::max()))
+			{
+				// Use byte for BP. It will use the correct picker and enum k2 node.
+				PinType.PinCategory = UEdGraphSchema_K2::PC_Byte;
+			}
+		}
+		else
+		{
+			checkf(false, TEXT("Typename in meta-data (%s) is invalid"), *InTargetType);
+		}
+	}
+
+	return PinType;
+}
+#endif
+
+} // UE::SceneState::Private
+
+bool IsPropertyReference(const FProperty* InProperty)
+{
+	const FStructProperty* StructProperty = CastField<FStructProperty>(InProperty);
+	return StructProperty && StructProperty->Struct && StructProperty->Struct->IsChildOf<FSceneStatePropertyReference>();
+}
+
+#if WITH_EDITOR
+bool IsPropertyReferenceCompatible(const FProperty* InReferenceProperty, const FProperty* InSourceProperty)
+{
+	check(InReferenceProperty && InSourceProperty);
+
+	if (!ensure(IsPropertyReference(InReferenceProperty)))
+	{
+		return false;
+	}
+
+	if (IsPropertyReference(InSourceProperty))
+	{
+		return Private::ArePropertyReferenceCompatible(InReferenceProperty, InSourceProperty);
+	}
+
+	const FProperty* TestProperty = InSourceProperty;
+
+	const bool bCanTargetRefArray = InReferenceProperty->HasMetaData(Metadata::CanRefToArray);
+	const bool bIsTargetRefArray = InReferenceProperty->HasMetaData(Metadata::IsRefToArray);
+
+	if (bIsTargetRefArray || bCanTargetRefArray)
+	{
+		if (const FArrayProperty* ArrayProperty = CastField<FArrayProperty>(TestProperty))
+		{
+			TestProperty = ArrayProperty->Inner;
+		}
+		else if(!bCanTargetRefArray)
+		{
+			return false;
+		}
+	}
+
+	const FString TargetType = InReferenceProperty->GetMetaData(Metadata::RefType);
+	return Private::IsPropertyCompatible(TargetType, TestProperty, bIsTargetRefArray);
+}
+
+FEdGraphPinType GetPropertyReferencePinType(const FProperty* InProperty)
+{
+	if (!ensure(InProperty))
+	{
+		return FEdGraphPinType();
+	}
+
+	ensure(IsPropertyReference(InProperty));
+
+	const EPinContainerType ContainerType = InProperty->HasMetaData(Metadata::IsRefToArray)
+		? EPinContainerType::Array
+		: EPinContainerType::None;
+
+	const FString TargetType = InProperty->GetMetaData(Metadata::RefType);
+	return Private::GetTargetPinType(TargetType, ContainerType);
+}
+#endif
+
+} // UE::SceneState
\ No newline at end of file
diff --git a/Engine/Plugins/Experimental/SceneState/Source/SceneStateBinding/Public/SceneStateBindingCollection.h b/Engine/Plugins/Experimental/SceneState/Source/SceneStateBinding/Public/SceneStateBindingCollection.h
--- a/Engine/Plugins/Experimental/SceneState/Source/SceneStateBinding/Public/SceneStateBindingCollection.h
+++ b/Engine/Plugins/Experimental/SceneState/Source/SceneStateBinding/Public/SceneStateBindingCollection.h
@@ -5,8 +5,11 @@
 #include "PropertyBindingBindingCollection.h"
 #include "SceneStateBinding.h"
 #include "SceneStateBindingDesc.h"
+#include "SceneStateBindingReference.h"
 #include "SceneStateBindingCollection.generated.h"
 
+struct FSceneStatePropertyReference;
+
 namespace UE::SceneState::Editor
 {
 	class FBindingCompiler;
@@ -30,8 +33,14 @@ struct FSceneStateBindingCollection : public FPropertyBindingBindingCollection
 	/** Finds the binding desc matching the given data handle */
 	SCENESTATEBINDING_API const FSceneStateBindingDesc* FindBindingDesc(FSceneStateBindingDataHandle InDataHandle) const;
 
-#if WITH_EDITOR
+	/** Retrieves the resolved reference for the given property reference */
+	SCENESTATEBINDING_API const FSceneStateBindingResolvedReference* FindResolvedReference(const FSceneStatePropertyReference& InPropertyReference) const;
+
+	/** Resolves the property access for the given reference and data view */
+	SCENESTATEBINDING_API uint8* ResolveProperty(const FSceneStateBindingResolvedReference& InResolvedReference, FPropertyBindingDataView InDataView) const;
+
 	//~ Begin FPropertyBindingBindingCollection
+#if WITH_EDITOR
 	SCENESTATEBINDING_API virtual FPropertyBindingBinding* AddBindingInternal(const FPropertyBindingPath& InSourcePath, const FPropertyBindingPath& InTargetPath) override;
 	SCENESTATEBINDING_API virtual void RemoveBindingsInternal(TFunctionRef<bool(FPropertyBindingBinding&)> InPredicate) override;
 	SCENESTATEBINDING_API virtual bool HasBindingInternal(TFunctionRef<bool(const FPropertyBindingBinding&)> InPredicate) const override;
@@ -47,14 +56,27 @@ struct FSceneStateBindingCollection : public FPropertyBindingBindingCollection
 	SCENESTATEBINDING_API virtual void VisitMutableBindings(TFunctionRef<EVisitResult(FPropertyBindingBinding& Binding)> InFunction) override;
 	SCENESTATEBINDING_API virtual void OnReset() override;
 	SCENESTATEBINDING_API virtual void VisitSourceStructDescriptorInternal(TFunctionRef<EVisitResult(const FPropertyBindingBindableStructDescriptor&)> InFunction) const override;
+	[[nodiscard]] SCENESTATEBINDING_API virtual bool OnResolvingPaths() override;
 	//~ End FPropertyBindingBindingCollection
 
 private:
 	UPROPERTY()
 	TArray<FSceneStateBindingDesc> BindingDescs;
 
+	/** Contains all the bindings from a source data to a target data */
 	UPROPERTY()
 	TArray<FSceneStateBinding> Bindings;
 
+	/**
+	 * Contains all the reference bindings from a source data.
+	 * This gets set and populated during scene state compilation.
+	 */
+	UPROPERTY()
+	TArray<FSceneStateBindingReference> References;
+
+	/** Contains the references that have been resolved for use at runtime */
+	UPROPERTY()
+	TArray<FSceneStateBindingResolvedReference> ResolvedReferences;
+
 	friend UE::SceneState::Editor::FBindingCompiler;
 };
diff --git a/Engine/Plugins/Experimental/SceneState/Source/SceneStateBinding/Public/SceneStateBindingDataHandle.h b/Engine/Plugins/Experimental/SceneState/Source/SceneStateBinding/Public/SceneStateBindingDataHandle.h
--- a/Engine/Plugins/Experimental/SceneState/Source/SceneStateBinding/Public/SceneStateBindingDataHandle.h
+++ b/Engine/Plugins/Experimental/SceneState/Source/SceneStateBinding/Public/SceneStateBindingDataHandle.h
@@ -74,6 +74,14 @@ struct FSceneStateBindingDataHandle
 		return static_cast<uint64>(Type) << 32 | static_cast<uint64>(Index) << 16 | static_cast<uint64>(SubIndex);
 	}
 
+	friend uint32 GetTypeHash(const FSceneStateBindingDataHandle& InDataHandle)
+	{
+		uint32 Hash = GetTypeHash(InDataHandle.Type);
+		Hash = HashCombineFast(Hash, GetTypeHash(InDataHandle.Index));
+		Hash = HashCombineFast(Hash, GetTypeHash(InDataHandle.SubIndex));
+		return Hash;
+	}
+
 private:
 	static constexpr uint16 InvalidIndex = std::numeric_limits<uint16>::max();
 
diff --git a/Engine/Plugins/Experimental/SceneState/Source/SceneStateBinding/Public/SceneStateBindingReference.h b/Engine/Plugins/Experimental/SceneState/Source/SceneStateBinding/Public/SceneStateBindingReference.h
new file mode 100644
--- /dev/null
+++ b/Engine/Plugins/Experimental/SceneState/Source/SceneStateBinding/Public/SceneStateBindingReference.h
@@ -0,0 +1,43 @@
+// Copyright Epic Games, Inc. All Rights Reserved.
+
+#pragma once
+
+#include "PropertyBindingPath.h"
+#include "SceneStateBindingDataHandle.h"
+#include "SceneStateBindingReference.generated.h"
+
+USTRUCT()
+struct FSceneStateBindingReference
+{
+	GENERATED_BODY()
+
+	/** Source property path of the reference */
+	UPROPERTY()
+	FPropertyBindingPath SourcePropertyPath;
+
+	/** Describes how to get the source data pointer */
+	UPROPERTY()
+	FSceneStateBindingDataHandle SourceDataHandle;
+};
+
+/** A resolved FSceneStateReference that contains more information about the source */
+USTRUCT()
+struct FSceneStateBindingResolvedReference
+{
+	GENERATED_BODY()
+
+	/** Source property access. */
+	UPROPERTY()
+	FPropertyBindingPropertyIndirection SourceIndirection;
+
+	/** Cached pointer to the leaf property of the access. */
+	const FProperty* SourceLeafProperty = nullptr;
+
+	/** Type of the source data, used for validation. */
+	UPROPERTY(Transient)
+	TObjectPtr<const UStruct> SourceStructType = nullptr;
+
+	/** Describes how to get the source data pointer. */
+	UPROPERTY()
+	FSceneStateBindingDataHandle SourceDataHandle;
+};
diff --git a/Engine/Plugins/Experimental/SceneState/Source/SceneStateBinding/Public/SceneStatePropertyReference.h b/Engine/Plugins/Experimental/SceneState/Source/SceneStateBinding/Public/SceneStatePropertyReference.h
new file mode 100644
--- /dev/null
+++ b/Engine/Plugins/Experimental/SceneState/Source/SceneStateBinding/Public/SceneStatePropertyReference.h
@@ -0,0 +1,36 @@
+// Copyright Epic Games, Inc. All Rights Reserved.
+
+#pragma once
+
+#include <limits>
+#include "SceneStatePropertyReference.generated.h"
+
+/**
+ * Property References allows the binding system to handle a property by reference rather than by copy.
+ * This is useful for tasks that need to write to external properties (e.g. as a form of output).
+ * The expected type of the reference should be set in "RefType" meta specifier.
+ * Meta specifiers for the type:
+ *
+ * RefType="<type>"
+ * Specifies the type of property that the reference supports.
+ * Supported types are: bool, byte, int32, int64, float, double, Name, String, Text, Wildcard, UObject pointers, and structs.
+ * Wildcards is a 'bypass' so that any valid FProperty is accepted, and is up to implementer to handle this.
+ * Structs and Objects must use full path name.
+ *
+ * IsRefToArray
+ * Specified to mean that the reference is to an TArray<RefType>
+ *
+ * CanRefToArray
+ * Specified to mean that the reference can bind to a RefType or TArray<RefType>
+ */
+USTRUCT()
+struct FSceneStatePropertyReference
+{
+	GENERATED_BODY()
+
+	/** Max number of indices that can be used */
+	static constexpr uint16 IndexCapacity = std::numeric_limits<uint16>::max();
+
+	UPROPERTY()
+	uint16 ReferenceIndex = IndexCapacity;
+};
diff --git a/Engine/Plugins/Experimental/SceneState/Source/SceneStateBinding/Public/SceneStatePropertyReferenceMetadata.h b/Engine/Plugins/Experimental/SceneState/Source/SceneStateBinding/Public/SceneStatePropertyReferenceMetadata.h
new file mode 100644
--- /dev/null
+++ b/Engine/Plugins/Experimental/SceneState/Source/SceneStateBinding/Public/SceneStatePropertyReferenceMetadata.h
@@ -0,0 +1,14 @@
+// Copyright Epic Games, Inc. All Rights Reserved.
+
+#pragma once
+
+#include "UObject/NameTypes.h"
+
+namespace UE::SceneState::Metadata
+{
+#if WITH_EDITOR
+	const FLazyName IsRefToArray(TEXT("IsRefToArray"));
+	const FLazyName CanRefToArray(TEXT("CanRefToArray"));
+	const FLazyName RefType(TEXT("RefType"));
+#endif
+}
diff --git a/Engine/Plugins/Experimental/SceneState/Source/SceneStateBinding/Public/SceneStatePropertyReferenceUtils.h b/Engine/Plugins/Experimental/SceneState/Source/SceneStateBinding/Public/SceneStatePropertyReferenceUtils.h
new file mode 100644
--- /dev/null
+++ b/Engine/Plugins/Experimental/SceneState/Source/SceneStateBinding/Public/SceneStatePropertyReferenceUtils.h
@@ -0,0 +1,159 @@
+// Copyright Epic Games, Inc. All Rights Reserved.
+
+#pragma once
+
+#include "Concepts/BaseStructureProvider.h"
+#include "UObject/EnumProperty.h"
+#include "UObject/TextProperty.h"
+#include "UObject/UnrealType.h"
+
+struct FEdGraphPinType;
+
+namespace UE::SceneState
+{
+	/** Returns whether the given property is a struct property of struct FSceneStatePropertyReference */
+	SCENESTATEBINDING_API bool IsPropertyReference(const FProperty* InProperty);
+
+#if WITH_EDITOR
+	/** Returns whether a source property is compatible with the reference property */
+	SCENESTATEBINDING_API bool IsPropertyReferenceCompatible(const FProperty* InReferenceProperty, const FProperty* InSourceProperty);
+
+	/** Returns the graph pin types based on a given Property Reference property's meta-data */
+	SCENESTATEBINDING_API FEdGraphPinType GetPropertyReferencePinType(const FProperty* InProperty);
+#endif
+
+	template<class InDataType>
+	struct TDataTypeHelper
+	{
+	};
+
+	template<>
+	struct TDataTypeHelper<void>
+	{
+		static bool IsValid(const FProperty* InProperty)
+		{
+			return true;
+		}
+	};
+
+	template<>
+	struct TDataTypeHelper<FBoolProperty::TCppType>
+	{
+		static bool IsValid(const FProperty* InProperty)
+		{
+			return InProperty && InProperty->IsA<FBoolProperty>();
+		};
+	};
+
+	template<>
+	struct TDataTypeHelper<FByteProperty::TCppType>
+	{
+		static bool IsValid(const FProperty* InProperty)
+		{
+			return InProperty && InProperty->IsA<FByteProperty>();
+		};
+	};
+
+	template<>
+	struct TDataTypeHelper<FIntProperty::TCppType>
+	{
+		static bool IsValid(const FProperty* InProperty)
+		{
+			return InProperty && InProperty->IsA<FIntProperty>();
+		};
+	};
+
+	template<>
+	struct TDataTypeHelper<FInt64Property::TCppType>
+	{
+		static bool IsValid(const FProperty* InProperty)
+		{
+			return InProperty && InProperty->IsA<FInt64Property>();
+		};
+	};
+
+	template<>
+	struct TDataTypeHelper<FFloatProperty::TCppType>
+	{
+		static bool IsValid(const FProperty* InProperty)
+		{
+			return InProperty && InProperty->IsA<FFloatProperty>();
+		};
+	};
+
+	template<>
+	struct TDataTypeHelper<FDoubleProperty::TCppType>
+	{
+		static bool IsValid(const FProperty* InProperty)
+		{
+			return InProperty && InProperty->IsA<FDoubleProperty>();
+		};
+	};
+
+	template<>
+	struct TDataTypeHelper<FNameProperty::TCppType>
+	{
+		static bool IsValid(const FProperty* InProperty)
+		{
+			return InProperty && InProperty->IsA<FNameProperty>();
+		};
+	};
+
+	template<>
+	struct TDataTypeHelper<FStrProperty::TCppType>
+	{
+		static bool IsValid(const FProperty* InProperty)
+		{
+			return InProperty && InProperty->IsA<FStrProperty>();
+		};
+	};
+
+	template<>
+	struct TDataTypeHelper<FTextProperty::TCppType>
+	{
+		static bool IsValid(const FProperty* InProperty)
+		{
+			return InProperty && InProperty->IsA<FTextProperty>();
+		};
+	};
+
+	template<class InDataType UE_REQUIRES(TIsTArray_V<InDataType>)>
+	struct TDataTypeHelper<InDataType>
+	{
+		static bool IsValid(const FProperty* InProperty)
+		{
+			const FArrayProperty* ArrayProperty = CastField<FArrayProperty>(InProperty);
+			return ArrayProperty && TDataTypeHelper<typename InDataType::ElementType>::IsValid(*ArrayProperty->Inner);
+		}
+	};
+
+	template<class InDataType UE_REQUIRES(TModels_V<CBaseStructureProvider, InDataType>)>
+	struct TDataTypeHelper<InDataType>
+	{
+		static bool IsValid(const FProperty* InProperty)
+		{
+			const FStructProperty* StructProperty = CastField<FStructProperty>(InProperty);
+			return StructProperty && StructProperty->Struct && StructProperty->Struct->IsChildOf(TBaseStructure<InDataType>::Get());
+		}
+	};
+
+	template<class InDataType UE_REQUIRES(TModels_V<CStaticClassProvider, typename TRemovePointer<InDataType>::Type>)>
+	struct TDataTypeHelper<InDataType>
+	{
+		static bool IsValid(const FProperty* InProperty)
+		{
+			const FObjectProperty* ObjectProperty = CastField<FObjectProperty>(InProperty);
+			return ObjectProperty && ObjectProperty->PropertyClass && ObjectProperty->PropertyClass->IsChildOf(TRemovePointer<InDataType>::Type::StaticClass());
+		}
+	};
+
+	template<class InDataType UE_REQUIRES(TIsEnum<InDataType>::Value)>
+	struct TDataTypeHelper<InDataType>
+	{
+		static bool IsValid(const FProperty* InProperty)
+		{
+			const FEnumProperty* EnumProperty = CastField<FEnumProperty>(InProperty);
+			return EnumProperty && EnumProperty->GetEnum() == StaticEnum<InDataType>();
+		}
+	};
+}
diff --git a/Engine/Plugins/Experimental/SceneState/Source/SceneStateBinding/SceneStateBinding.Build.cs b/Engine/Plugins/Experimental/SceneState/Source/SceneStateBinding/SceneStateBinding.Build.cs
--- a/Engine/Plugins/Experimental/SceneState/Source/SceneStateBinding/SceneStateBinding.Build.cs
+++ b/Engine/Plugins/Experimental/SceneState/Source/SceneStateBinding/SceneStateBinding.Build.cs
@@ -23,5 +23,15 @@ public class SceneStateBinding : ModuleRules
                 "Engine",
             }
         );
+
+        if (Target.Type == TargetRules.TargetType.Editor)
+        {
+            PrivateDependencyModuleNames.AddRange(
+                new string[]
+                {
+                    "BlueprintGraph",
+                }
+            );
+        }
     }
 }
\ No newline at end of file
diff --git a/Engine/Plugins/Experimental/SceneState/Source/SceneStateBlueprintEditor/Private/Compiler/SceneStateBindingCompiler.cpp b/Engine/Plugins/Experimental/SceneState/Source/SceneStateBlueprintEditor/Private/Compiler/SceneStateBindingCompiler.cpp
--- a/Engine/Plugins/Experimental/SceneState/Source/SceneStateBlueprintEditor/Private/Compiler/SceneStateBindingCompiler.cpp
+++ b/Engine/Plugins/Experimental/SceneState/Source/SceneStateBlueprintEditor/Private/Compiler/SceneStateBindingCompiler.cpp
@@ -1,10 +1,13 @@
 // Copyright Epic Games, Inc. All Rights Reserved.
 
 #include "SceneStateBindingCompiler.h"
+#include "Misc/EnumerateRange.h"
+#include "SceneStateBindingUtils.h"
 #include "SceneStateBlueprint.h"
 #include "SceneStateBlueprintCompilerContext.h"
 #include "SceneStateEventSchema.h"
 #include "SceneStateGeneratedClass.h"
+#include "SceneStatePropertyReference.h"
 #include "StructUtils/UserDefinedStruct.h"
 #include "Tasks/SceneStateBlueprintableTask.h"
 #include "Tasks/SceneStateTaskBindingExtension.h"
@@ -14,6 +17,46 @@
 namespace UE::SceneState::Editor
 {
 
+namespace Private
+{
+
+/** Returns whether the given BindingDesc+Path is a struct property of the given struct type */
+bool IsStructPropertyOfType(const UStruct* InStruct, const FSceneStateBindingDesc& InBindingDesc, const FPropertyBindingPath& InBindingPath)
+{
+	TArray<FPropertyBindingPathIndirection> Indirections;
+	if (InBindingPath.ResolveIndirections(InBindingDesc.Struct, Indirections) && !Indirections.IsEmpty())
+	{
+		const FProperty* TargetProperty = Indirections.Last().GetProperty();
+		check(TargetProperty);
+		if (const FStructProperty* TargetStructProperty = CastField<FStructProperty>(TargetProperty))
+		{
+			return TargetStructProperty->Struct->IsChildOf(InStruct);
+		}
+	}
+	return false;
+}
+
+/**
+ * Returns the preferred batching type for a given binding
+ * @param InBindingCollection the collection containing the binding to check
+ * @param InBinding the binding to check
+ * @return the type of batching that should be performed for the given binding
+ */
+EDataAccessType GetBindingBatchingType(const FSceneStateBindingCollection& InBindingCollection, const FSceneStateBinding& InBinding)
+{
+	if (const FSceneStateBindingDesc* TargetDesc = InBindingCollection.FindBindingDesc(InBinding.TargetDataHandle))
+	{
+		if (IsStructPropertyOfType(FSceneStatePropertyReference::StaticStruct(), *TargetDesc, InBinding.GetTargetPath()))
+		{
+			return EDataAccessType::Reference;
+		}
+		return EDataAccessType::Copy;
+	}
+	return EDataAccessType::None;
+}
+
+} // UE::SceneState::Editor::Private
+	
 void FBindingCompiler::Compile()
 {
 	check(Blueprint && Class);
@@ -36,7 +79,15 @@ void FBindingCompiler::Compile()
 	RemoveInvalidBindings();
 
 	ResolveBindingDataHandles();
-	BatchCopies();
+
+	GroupBindings();
+	CompileCopies();
+	CompileReferences();
+}
+
+void FBindingCompiler::AddDataView(const FSceneStateBindingDataHandle& InDataHandle, FPropertyBindingDataView InDataView)
+{
+	TemplateDataViewMap.Add(InDataHandle, InDataView);
 }
 
 void FBindingCompiler::AddBindingDesc(FSceneStateBindingDesc&& InBindingDesc)
@@ -109,7 +160,7 @@ void FBindingCompiler::AddTaskBindingDescs()
 	{
 		const FSceneStateTask& Task = Class->Tasks[TaskIndex].Get<FSceneStateTask>();
 		const FSceneStateTaskMetadata& TaskMetadata = Class->TaskMetadata[TaskIndex];
-		const FConstStructView TaskInstance = Class->TaskInstances[TaskIndex];
+		const FStructView TaskInstance = Class->TaskInstances[TaskIndex];
 
 		// Task Instance Binding
 		{
@@ -119,6 +170,7 @@ void FBindingCompiler::AddTaskBindingDescs()
 			BindingDesc.Struct = TaskInstance.GetScriptStruct();
 			BindingDesc.DataHandle = FSceneStateBindingDataHandle(ESceneStateDataType::Task, TaskIndex);
 
+			AddDataView(BindingDesc.DataHandle, TaskInstance);
 			AddBindingDesc(MoveTemp(BindingDesc));
 		}
 
@@ -239,24 +291,71 @@ void FBindingCompiler::RemoveInvalidBindings()
 	Class->BindingCollection.RemoveInvalidBindings(ValidBindingMap);
 }
 
-void FBindingCompiler::BatchCopies()
+void FBindingCompiler::GroupBindings()
 {
-	// To batch first sort the bindings so all the same target data handles are together
-	Class->BindingCollection.Bindings.StableSort(
-		[](const FSceneStateBinding& InLeft, const FSceneStateBinding& InRight)
+	FSceneStateBindingCollection& BindingCollection = Class->BindingCollection;
+
+	// Sort by Batching Type (Copy,Reference) then by the Target Data.
+	BindingCollection.Bindings.StableSort(
+		[&BindingCollection](const FSceneStateBinding& InLeft, const FSceneStateBinding& InRight)
 		{
+			const EDataAccessType LeftBatchingType = Private::GetBindingBatchingType(BindingCollection, InLeft);
+			const EDataAccessType RightBatchingType = Private::GetBindingBatchingType(BindingCollection, InRight);
+
+			if (LeftBatchingType != RightBatchingType)
+			{
+				return LeftBatchingType < RightBatchingType;
+			}
+
 			return InLeft.TargetDataHandle.AsNumber() < InRight.TargetDataHandle.AsNumber();
 		});
 
+	BatchRangeMap.Reset();
+
+	// Keeps track of the latest batching type to determine its range
+	TOptional<EDataAccessType> LatestBatchingType;
+
+	for (TConstEnumerateRef<FSceneStateBinding> Binding : EnumerateRange(BindingCollection.Bindings))
+	{
+		const EDataAccessType BatchingType = Private::GetBindingBatchingType(BindingCollection, *Binding);
+		if (LatestBatchingType.IsSet() && BatchingType == *LatestBatchingType)
+		{
+			// Increment the count for the latest batch range
+			++BatchRangeMap[*LatestBatchingType].Count;
+		}
+		else
+		{
+			LatestBatchingType = BatchingType;
+
+			// Batch Start Index Map must not contain batching types that have been previously added as the bindings were already sorted by batching types 
+			ensure(!BatchRangeMap.Contains(BatchingType));
+
+			FBatchRange& BatchRange = BatchRangeMap.Add(BatchingType);
+			BatchRange.Index = Binding.GetIndex();
+			BatchRange.Count = 1;
+		}
+	}
+}
+
+void FBindingCompiler::CompileCopies()
+{
+	const FBatchRange* CopyRange = BatchRangeMap.Find(EDataAccessType::Copy);
+	if (!CopyRange)
+	{
+		return;
+	}
+
+	TConstArrayView<FSceneStateBinding> CopyBindings = MakeArrayView(Class->BindingCollection.Bindings).Slice(CopyRange->Index, CopyRange->Count);
+
 	// Add 1 batch per group of bindings that all have the same target data handle
-	for (TArray<FSceneStateBinding>::TIterator BindingIter(Class->BindingCollection.Bindings); BindingIter;)
+	for (int32 BindingIndex = 0; BindingIndex < CopyBindings.Num();)
 	{
-		const FSceneStateBindingDataHandle TargetDataHandle = BindingIter->TargetDataHandle;
+		const FSceneStateBinding& CopyBinding = CopyBindings[BindingIndex];
 
-		const FSceneStateBindingDesc* TargetDesc = Class->BindingCollection.FindBindingDesc(TargetDataHandle);
+		const FSceneStateBindingDesc* TargetDesc = Class->BindingCollection.FindBindingDesc(CopyBinding.TargetDataHandle);
 		if (!ensure(TargetDesc))
 		{
-			++BindingIter;
+			++BindingIndex;
 			continue;
 		}
 
@@ -264,16 +363,69 @@ void FBindingCompiler::BatchCopies()
 
 		FPropertyBindingCopyInfoBatch& CopyBatch = Class->BindingCollection.AddCopyBatch();
 		CopyBatch.TargetStruct = TInstancedStruct<FSceneStateBindingDesc>::Make(*TargetDesc);
-		CopyBatch.BindingsBegin = FPropertyBindingIndex16(BindingIter.GetIndex());
+		CopyBatch.BindingsBegin = FPropertyBindingIndex16(BindingIndex);
 
 		// Get the last index that matches the current target data handle
-		for (; BindingIter && BindingIter->TargetDataHandle == TargetDataHandle; ++BindingIter)
+		for (; BindingIndex < CopyBindings.Num() && CopyBindings[BindingIndex].TargetDataHandle == CopyBinding.TargetDataHandle; ++BindingIndex)
 		{
 		}
 
 		// Binding end is one past the last binding.
-		CopyBatch.BindingsEnd = FPropertyBindingIndex16(BindingIter.GetIndex());
-		OnBindingsBatchCompiled(BatchBindings, TargetDataHandle);
+		CopyBatch.BindingsEnd = FPropertyBindingIndex16(BindingIndex);
+		OnBindingsBatchCompiled(BatchBindings, CopyBinding.TargetDataHandle);
+	}
+}
+
+void FBindingCompiler::CompileReferences()
+{
+	const FBatchRange* ReferenceRange = BatchRangeMap.Find(EDataAccessType::Reference);
+	if (!ReferenceRange)
+	{
+		return;
+	}
+
+	if (ReferenceRange->Count > FSceneStatePropertyReference::IndexCapacity)
+	{
+		const FText ErrorMessageFormat = FText::Format(LOCTEXT("MaxReferenceCountReached"
+			, "Reference range {0} has more entries than the allowed number {1}")
+			, FText::AsNumber(ReferenceRange->Count)
+			, FText::AsNumber(FSceneStatePropertyReference::IndexCapacity));
+
+		Context.MessageLog.Error(*ErrorMessageFormat.ToString());
+		return;
+	}
+
+	TConstArrayView<FSceneStateBinding> ReferenceBindings = MakeArrayView(Class->BindingCollection.Bindings).Slice(ReferenceRange->Index, ReferenceRange->Count);
+
+	Class->BindingCollection.References.Empty(ReferenceBindings.Num());
+
+	TArray<FPropertyBindingPathIndirection> TargetIndirections;
+	FString ErrorMessage;
+
+	for (const FSceneStateBinding& ReferenceBinding : ReferenceBindings)
+	{
+		const FPropertyBindingDataView TemplateDataView = TemplateDataViewMap[ReferenceBinding.TargetDataHandle];
+
+		TargetIndirections.Reset();
+		if (!ReferenceBinding.GetTargetPath().ResolveIndirectionsWithValue(TemplateDataView, TargetIndirections, &ErrorMessage))
+		{
+			const FText ErrorMessageFormat = FText::Format(LOCTEXT("FailedResolvingReferenceIndirection"
+				, "Error resolving path in {0} ({1}): {2}")
+				, FText::FromString(GetNameSafe(TemplateDataView.GetStruct()))
+				, FText::FromString(ReferenceBinding.ToString())
+				, FText::FromString(ErrorMessage));
+
+			Context.MessageLog.Error(*ErrorMessageFormat.ToString());
+			continue;
+		}
+
+		FSceneStatePropertyReference* PropertyReference = static_cast<FSceneStatePropertyReference*>(const_cast<void*>(TargetIndirections.Last().GetPropertyAddress()));
+		check(PropertyReference);
+		PropertyReference->ReferenceIndex = Class->BindingCollection.References.AddDefaulted();
+
+		FSceneStateBindingReference& Reference =  Class->BindingCollection.References[PropertyReference->ReferenceIndex];
+		Reference.SourcePropertyPath = ReferenceBinding.GetSourcePath();
+		Reference.SourceDataHandle = ReferenceBinding.SourceDataHandle;
 	}
 }
 
diff --git a/Engine/Plugins/Experimental/SceneState/Source/SceneStateBlueprintEditor/Private/Compiler/SceneStateBindingCompiler.h b/Engine/Plugins/Experimental/SceneState/Source/SceneStateBlueprintEditor/Private/Compiler/SceneStateBindingCompiler.h
--- a/Engine/Plugins/Experimental/SceneState/Source/SceneStateBlueprintEditor/Private/Compiler/SceneStateBindingCompiler.h
+++ b/Engine/Plugins/Experimental/SceneState/Source/SceneStateBlueprintEditor/Private/Compiler/SceneStateBindingCompiler.h
@@ -4,13 +4,13 @@
 
 #include "Misc/Guid.h"
 #include "PropertyBindingDataView.h"
+#include "SceneStateBindingDataHandle.h"
 #include "UObject/ObjectPtr.h"
 
 class USceneStateBlueprint;
 class USceneStateGeneratedClass;
 struct FPropertyBindingIndex16;
 struct FSceneStateBinding;
-struct FSceneStateBindingDataHandle;
 struct FSceneStateBindingDesc;
 
 namespace UE::SceneState::Editor
@@ -21,6 +21,14 @@ namespace UE::SceneState::Editor
 namespace UE::SceneState::Editor
 {
 
+/** Describes the types how a binding can be batched */
+enum class EDataAccessType : uint8
+{
+	None,
+	Copy,
+	Reference,
+};
+
 class FBindingCompiler
 {
 public:
@@ -34,6 +42,7 @@ public:
 	void Compile();
 
 private:
+	void AddDataView(const FSceneStateBindingDataHandle& InDataHandle, FPropertyBindingDataView InDataView);
 	void AddBindingDesc(FSceneStateBindingDesc&& InBindingDesc);
 
 	void AddRootBindingDesc();
@@ -46,7 +55,10 @@ private:
 
 	void ResolveBindingDataHandles();
 	void RemoveInvalidBindings();
-	void BatchCopies();
+
+	void GroupBindings();
+	void CompileCopies();
+	void CompileReferences();
 
 	void OnBindingsBatchCompiled(FPropertyBindingIndex16 InBindingsBatch, const FSceneStateBindingDataHandle& InTargetDataHandle);
 
@@ -63,6 +75,22 @@ private:
 
 	/** Map of struct ids to their binding data view */
 	TMap<FGuid, const FPropertyBindingDataView> ValidBindingMap;
+
+	/** Map of gathered binding data handle to the template data view */
+	TMap<FSceneStateBindingDataHandle, FPropertyBindingDataView> TemplateDataViewMap;
+
+	struct FBatchRange
+	{
+		/** Start index of the batch, inclusive */
+		int32 Index = INDEX_NONE;
+		/** Number of elements in the batch */
+		int32 Count = 0;
+	};
+	/**
+	 * Map of the batching type to the binding range index for that type
+	 * A map is not the best candidate to hold this information, and only used for readability.
+	 */
+	TMap<EDataAccessType, FBatchRange> BatchRangeMap;
 };
 
 } // UE::SceneState::Editor
diff --git a/Engine/Plugins/Experimental/SceneState/Source/SceneStateBlueprintEditor/Private/DetailsView/SceneStateCachedBindingData.cpp b/Engine/Plugins/Experimental/SceneState/Source/SceneStateBlueprintEditor/Private/DetailsView/SceneStateCachedBindingData.cpp
new file mode 100644
--- /dev/null
+++ b/Engine/Plugins/Experimental/SceneState/Source/SceneStateBlueprintEditor/Private/DetailsView/SceneStateCachedBindingData.cpp
@@ -0,0 +1,100 @@
+// Copyright Epic Games, Inc. All Rights Reserved.
+
+#include "SceneStateCachedBindingData.h"
+#include "IPropertyAccessEditor.h"
+#include "PropertyBindingDataView.h"
+#include "SceneStateBindingDesc.h"
+#include "SceneStatePropertyReference.h"
+#include "SceneStatePropertyReferenceMetadata.h"
+#include "SceneStatePropertyReferenceUtils.h"
+
+#define LOCTEXT_NAMESPACE "SceneStateCachedBindingData"
+
+namespace UE::SceneState::Editor
+{
+
+void FSceneStateCachedBindingData::UpdatePropertyReferenceTooltip(const FProperty& InProperty, FTextBuilder& InOutTextBuilder) const
+{
+	if (InProperty.HasMetaData(Metadata::IsRefToArray))
+	{
+		InOutTextBuilder.AppendLineFormat(LOCTEXT("PropertyRefBindingTooltipArray", "Supported types are Array of {0}")
+			, FText::FromString(InProperty.GetMetaData(Metadata::RefType)));
+	}
+	else
+	{
+		InOutTextBuilder.AppendLineFormat(LOCTEXT("PropertyRefBindingTooltip", "Supported types are {0}")
+			, FText::FromString(InProperty.GetMetaData(Metadata::RefType)));
+
+		if (InProperty.HasMetaData(Metadata::CanRefToArray))
+		{
+			InOutTextBuilder.AppendLine(LOCTEXT("PropertyRefBindingTooltipCanSupportArray", "Supports Arrays"));
+		}
+	}
+}
+
+bool FSceneStateCachedBindingData::GetPinTypeAndIconForProperty(const FProperty& InProperty, FPropertyBindingDataView InTargetDataView, FEdGraphPinType& OutPinType, FName& OutIconName) const
+{
+	if (UE::SceneState::IsPropertyReference(&InProperty))
+	{
+		OutPinType = GetPropertyReferencePinType(&InProperty);
+		OutIconName = TEXT("Kismet.Tabs.Variables");
+		return true;
+	}
+	return false;
+}
+
+bool FSceneStateCachedBindingData::IsPropertyReference(const FProperty& InProperty)
+{
+	return UE::SceneState::IsPropertyReference(&InProperty);
+}
+
+void FSceneStateCachedBindingData::AddPropertyInfoOverride(const FProperty& InProperty, TArray<TSharedPtr<const UE::PropertyBinding::FPropertyInfoOverride>>& OutPropertyInfoOverrides) const
+{
+	const FStructProperty* StructProperty = CastField<const FStructProperty>(&InProperty);
+	if (!StructProperty || !StructProperty->Struct)
+	{
+		return;
+	}
+
+	// Add the PropertyRef property type with its RefType
+	if (StructProperty->Struct->IsChildOf<FSceneStatePropertyReference>())
+	{
+		FEdGraphPinType PinType = GetPropertyReferencePinType(&InProperty);
+
+		// If Property supports Arrays, add the Array version of these pin types
+		if (GetPropertyHandle()->HasMetaData(Metadata::CanRefToArray))
+		{
+			PinType.ContainerType = EPinContainerType::Array;
+		}
+
+		TSharedRef<UE::PropertyBinding::FPropertyInfoOverride> PropertyInfoOverride = MakeShared<UE::PropertyBinding::FPropertyInfoOverride>();
+		PropertyInfoOverride->PinType = PinType;
+
+		FString TypeName;
+		if (UObject* SubCategoryObject = PinType.PinSubCategoryObject.Get())
+		{
+			TypeName = SubCategoryObject->GetName();
+		}
+		else
+		{
+			TypeName = PinType.PinCategory.ToString() + TEXT(" ") + PinType.PinSubCategory.ToString();
+		}
+
+		PropertyInfoOverride->TypeNameText = FText::FromString(TypeName);
+		OutPropertyInfoOverrides.Emplace(MoveTemp(PropertyInfoOverride));
+	}
+}
+
+bool FSceneStateCachedBindingData::DeterminePropertiesCompatibilityInternal(const FProperty* InSourceProperty, const FProperty* InTargetProperty, const void* InSourcePropertyValue, const void* InTargetPropertyValue, bool& bOutAreCompatible) const
+{
+	if (UE::SceneState::IsPropertyReference(InTargetProperty))
+	{
+		bOutAreCompatible = UE::SceneState::IsPropertyReferenceCompatible(InTargetProperty, InSourceProperty);
+		return true;
+	}
+	return false;
+}
+
+} // UE::SceneState::Editor
+
+#undef LOCTEXT_NAMESPACE
diff --git a/Engine/Plugins/Experimental/SceneState/Source/SceneStateBlueprintEditor/Private/DetailsView/SceneStateCachedBindingData.h b/Engine/Plugins/Experimental/SceneState/Source/SceneStateBlueprintEditor/Private/DetailsView/SceneStateCachedBindingData.h
new file mode 100644
--- /dev/null
+++ b/Engine/Plugins/Experimental/SceneState/Source/SceneStateBlueprintEditor/Private/DetailsView/SceneStateCachedBindingData.h
@@ -0,0 +1,24 @@
+// Copyright Epic Games, Inc. All Rights Reserved.
+
+#pragma once
+
+#include "PropertyBindingExtension.h"
+
+namespace UE::SceneState::Editor
+{
+
+struct FSceneStateCachedBindingData : UE::PropertyBinding::FCachedBindingData, TSharedFromThis<FSceneStateCachedBindingData>
+{
+	using Super = UE::PropertyBinding::FCachedBindingData;
+	using Super::Super;
+
+	//~ Begin UE::PropertyBinding::FCachedBindingData
+	virtual void UpdatePropertyReferenceTooltip(const FProperty& InProperty, FTextBuilder& InOutTextBuilder) const override;
+	virtual bool GetPinTypeAndIconForProperty(const FProperty& InProperty, FPropertyBindingDataView InTargetDataView, FEdGraphPinType& OutPinType, FName& OutIconName) const override;
+	virtual bool IsPropertyReference(const FProperty& InProperty) override;
+	virtual void AddPropertyInfoOverride(const FProperty& InProperty, TArray<TSharedPtr<const UE::PropertyBinding::FPropertyInfoOverride>>& OutPropertyInfoOverrides) const override;
+	virtual bool DeterminePropertiesCompatibilityInternal(const FProperty* InSourceProperty,const FProperty* InTargetProperty,const void* InSourcePropertyValue,const void* InTargetPropertyValue,bool& bOutAreCompatible) const override;
+	//~ End UE::PropertyBinding::FCachedBindingData
+};
+
+} // UE::SceneState::Editor
diff --git a/Engine/Plugins/Experimental/SceneState/Source/SceneStateBlueprintEditor/Private/DetailsView/SceneStatePropertyBindingExtension.cpp b/Engine/Plugins/Experimental/SceneState/Source/SceneStateBlueprintEditor/Private/DetailsView/SceneStatePropertyBindingExtension.cpp
--- a/Engine/Plugins/Experimental/SceneState/Source/SceneStateBlueprintEditor/Private/DetailsView/SceneStatePropertyBindingExtension.cpp
+++ b/Engine/Plugins/Experimental/SceneState/Source/SceneStateBlueprintEditor/Private/DetailsView/SceneStatePropertyBindingExtension.cpp
@@ -3,6 +3,8 @@
 #include "SceneStatePropertyBindingExtension.h"
 #include "GameFramework/Actor.h"
 #include "SceneStateBlueprintEditorUtils.h"
+#include "SceneStateCachedBindingData.h"
+#include "SceneStatePropertyReference.h"
 
 namespace UE::SceneState::Editor
 {
@@ -19,4 +21,30 @@ void FBindingExtension::ExtendWidgetRow(FDetailWidgetRow& InWidgetRow, const IDe
 	FPropertyBindingExtension::ExtendWidgetRow(InWidgetRow, InDetailBuilder, InObjectClass, InPropertyHandle);
 }
 
+TSharedPtr<UE::PropertyBinding::FCachedBindingData> FBindingExtension::CreateCachedBindingData(IPropertyBindingBindingCollectionOwner* InBindingsOwner
+	, const FPropertyBindingPath& InTargetPath
+	, const TSharedPtr<IPropertyHandle>& InPropertyHandle
+	, TConstArrayView<TInstancedStruct<FPropertyBindingBindableStructDescriptor>> InAccessibleStructs) const
+{
+	return MakeShared<FSceneStateCachedBindingData>(InBindingsOwner, InTargetPath, InPropertyHandle, InAccessibleStructs);
+}
+
+bool FBindingExtension::GetPromotionToParameterOverrideInternal(const FProperty& InProperty, bool& bOutOverride) const
+{
+	if (const FStructProperty* StructProperty = CastField<FStructProperty>(&InProperty))
+	{
+		// Support Property Refs as even though these aren't bp types, the actual types that would be added are the ones in the meta-data RefType
+		if (StructProperty->Struct && StructProperty->Struct->IsChildOf<FSceneStatePropertyReference>())
+		{
+			bOutOverride = false;
+			return true;
+		}
+	}
+	return false;
+}
+
+bool FBindingExtension::CanBindToArrayElements() const
+{
+	return true;
+}
 } // UE::SceneState::Editor
diff --git a/Engine/Plugins/Experimental/SceneState/Source/SceneStateBlueprintEditor/Private/DetailsView/SceneStatePropertyBindingExtension.h b/Engine/Plugins/Experimental/SceneState/Source/SceneStateBlueprintEditor/Private/DetailsView/SceneStatePropertyBindingExtension.h
--- a/Engine/Plugins/Experimental/SceneState/Source/SceneStateBlueprintEditor/Private/DetailsView/SceneStatePropertyBindingExtension.h
+++ b/Engine/Plugins/Experimental/SceneState/Source/SceneStateBlueprintEditor/Private/DetailsView/SceneStatePropertyBindingExtension.h
@@ -13,6 +13,12 @@ public:
 	//~ Begin IDetailPropertyExtensionHandler
 	virtual void ExtendWidgetRow(FDetailWidgetRow& InWidgetRow, const IDetailLayoutBuilder& InDetailBuilder, const UClass* InObjectClass, TSharedPtr<IPropertyHandle> InPropertyHandle) override;
 	//~ End IDetailPropertyExtensionHandler
+
+	//~ Begin FPropertyBindingExtension
+	virtual TSharedPtr<UE::PropertyBinding::FCachedBindingData> CreateCachedBindingData(IPropertyBindingBindingCollectionOwner* InBindingsOwner, const FPropertyBindingPath& InTargetPath, const TSharedPtr<IPropertyHandle>& InPropertyHandle, TConstArrayView<TInstancedStruct<FPropertyBindingBindableStructDescriptor>> InAccessibleStructs) const override;
+	virtual bool GetPromotionToParameterOverrideInternal(const FProperty& InProperty, bool& bOutOverride) const override;
+	virtual bool CanBindToArrayElements() const override;
+	//~ End FPropertyBindingExtension
 };
 
 } // UE::SceneState::Editor
