Motion Design Scene State: Added core setter task for Bool, String and Integer properties.

diff --git a/Engine/Plugins/Experimental/SceneState/Source/SceneStateTasks/Private/Setters/SceneStateSetBoolTask.cpp b/Engine/Plugins/Experimental/SceneState/Source/SceneStateTasks/Private/Setters/SceneStateSetBoolTask.cpp
new file mode 100644
--- /dev/null
+++ b/Engine/Plugins/Experimental/SceneState/Source/SceneStateTasks/Private/Setters/SceneStateSetBoolTask.cpp
@@ -0,0 +1,17 @@
+// Copyright Epic Games, Inc. All Rights Reserved.
+
+#include "Setters/SceneStateSetBoolTask.h"
+#include "Setters/SceneStateSetterUtils.h"
+
+#if WITH_EDITOR
+const UScriptStruct* FSceneStateSetBoolTask::OnGetTaskInstanceType() const
+{
+	return FInstanceDataType::StaticStruct();
+}
+#endif
+
+void FSceneStateSetBoolTask::OnStart(const FSceneStateExecutionContext& InContext, FStructView InTaskInstance) const
+{
+	UE::SceneState::SetValue(*this, InContext, InTaskInstance);
+	Finish(InContext, InTaskInstance);
+}
diff --git a/Engine/Plugins/Experimental/SceneState/Source/SceneStateTasks/Private/Setters/SceneStateSetIntegerTask.cpp b/Engine/Plugins/Experimental/SceneState/Source/SceneStateTasks/Private/Setters/SceneStateSetIntegerTask.cpp
new file mode 100644
--- /dev/null
+++ b/Engine/Plugins/Experimental/SceneState/Source/SceneStateTasks/Private/Setters/SceneStateSetIntegerTask.cpp
@@ -0,0 +1,17 @@
+// Copyright Epic Games, Inc. All Rights Reserved.
+
+#include "Setters/SceneStateSetIntegerTask.h"
+#include "Setters/SceneStateSetterUtils.h"
+
+#if WITH_EDITOR
+const UScriptStruct* FSceneStateSetIntegerTask::OnGetTaskInstanceType() const
+{
+	return FInstanceDataType::StaticStruct();
+}
+#endif
+
+void FSceneStateSetIntegerTask::OnStart(const FSceneStateExecutionContext& InContext, FStructView InTaskInstance) const
+{
+	UE::SceneState::SetValue(*this, InContext, InTaskInstance);
+	Finish(InContext, InTaskInstance);
+}
diff --git a/Engine/Plugins/Experimental/SceneState/Source/SceneStateTasks/Private/Setters/SceneStateSetStringTask.cpp b/Engine/Plugins/Experimental/SceneState/Source/SceneStateTasks/Private/Setters/SceneStateSetStringTask.cpp
new file mode 100644
--- /dev/null
+++ b/Engine/Plugins/Experimental/SceneState/Source/SceneStateTasks/Private/Setters/SceneStateSetStringTask.cpp
@@ -0,0 +1,17 @@
+// Copyright Epic Games, Inc. All Rights Reserved.
+
+#include "Setters/SceneStateSetStringTask.h"
+#include "Setters/SceneStateSetterUtils.h"
+
+#if WITH_EDITOR
+const UScriptStruct* FSceneStateSetStringTask::OnGetTaskInstanceType() const
+{
+	return FInstanceDataType::StaticStruct();
+}
+#endif
+
+void FSceneStateSetStringTask::OnStart(const FSceneStateExecutionContext& InContext, FStructView InTaskInstance) const
+{
+	UE::SceneState::SetValue(*this, InContext, InTaskInstance);
+	Finish(InContext, InTaskInstance);
+}
diff --git a/Engine/Plugins/Experimental/SceneState/Source/SceneStateTasks/Private/Setters/SceneStateSetTextTask.cpp b/Engine/Plugins/Experimental/SceneState/Source/SceneStateTasks/Private/Setters/SceneStateSetTextTask.cpp
new file mode 100644
--- /dev/null
+++ b/Engine/Plugins/Experimental/SceneState/Source/SceneStateTasks/Private/Setters/SceneStateSetTextTask.cpp
@@ -0,0 +1,17 @@
+// Copyright Epic Games, Inc. All Rights Reserved.
+
+#include "Setters/SceneStateSetTextTask.h"
+#include "Setters/SceneStateSetterUtils.h"
+
+#if WITH_EDITOR
+const UScriptStruct* FSceneStateSetTextTask::OnGetTaskInstanceType() const
+{
+	return FInstanceDataType::StaticStruct();
+}
+#endif
+
+void FSceneStateSetTextTask::OnStart(const FSceneStateExecutionContext& InContext, FStructView InTaskInstance) const
+{
+	UE::SceneState::SetValue(*this, InContext, InTaskInstance);
+	Finish(InContext, InTaskInstance);
+}
diff --git a/Engine/Plugins/Experimental/SceneState/Source/SceneStateTasks/Private/Setters/SceneStateSetterUtils.cpp b/Engine/Plugins/Experimental/SceneState/Source/SceneStateTasks/Private/Setters/SceneStateSetterUtils.cpp
new file mode 100644
--- /dev/null
+++ b/Engine/Plugins/Experimental/SceneState/Source/SceneStateTasks/Private/Setters/SceneStateSetterUtils.cpp
@@ -0,0 +1,31 @@
+// Copyright Epic Games, Inc. All Rights Reserved.
+
+#include "Setters/SceneStateSetterUtils.h"
+
+namespace UE::SceneState
+{
+
+bool CallValueSetter(uint8* InDestValuePtr, const FProperty* InDestProperty, const uint8* InSourceValuePtr, int32 InArrayIndex)
+{
+	if (!ensureMsgf(InArrayIndex >= 0 && InArrayIndex < InDestProperty->ArrayDim, TEXT("Index %d outside property array dim bounds %d!"), InArrayIndex, InDestProperty->ArrayDim))
+	{
+		return false;
+	}
+
+	// No setter, can't call via this method
+	if (!InDestProperty->HasSetter())
+	{
+		return false;
+	}
+
+	// Reverse logic of ContainerPtrToValuePtr (i.e. ValuePtrToContainerPtr)
+	uint8* const DestContainerPtr = InDestValuePtr - (InDestProperty->GetOffset_ForInternal() + InDestProperty->GetElementSize() * InArrayIndex);
+
+	// check that this reverse logic remains 'up-to-date' with the latest ContainerPtrToValuePtr logic
+	checkf(InDestProperty->ContainerPtrToValuePtr<uint8>(DestContainerPtr) == InDestValuePtr, TEXT("ValuePtrToContainerPtr logic outdated. Verify against ContainerPtrToValuePtr."));
+
+	InDestProperty->CallSetter(DestContainerPtr, InSourceValuePtr);
+	return true;
+}
+
+} // UE::SceneState
diff --git a/Engine/Plugins/Experimental/SceneState/Source/SceneStateTasks/Public/Setters/SceneStateSetBoolTask.h b/Engine/Plugins/Experimental/SceneState/Source/SceneStateTasks/Public/Setters/SceneStateSetBoolTask.h
new file mode 100644
--- /dev/null
+++ b/Engine/Plugins/Experimental/SceneState/Source/SceneStateTasks/Public/Setters/SceneStateSetBoolTask.h
@@ -0,0 +1,39 @@
+// Copyright Epic Games, Inc. All Rights Reserved.
+
+#pragma once
+
+#include "SceneStatePropertyReference.h"
+#include "Tasks/SceneStateTask.h"
+#include "Tasks/SceneStateTaskInstance.h"
+#include "SceneStateSetBoolTask.generated.h"
+
+USTRUCT()
+struct FSceneStateSetBoolTaskInstance : public FSceneStateTaskInstance
+{
+	GENERATED_BODY()
+
+	/** Target bool property to set */
+	UPROPERTY(EditAnywhere, Category="Setter", meta=(RefType="bool"))
+	FSceneStatePropertyReference Target;
+
+	/** The boolean value to set. */
+	UPROPERTY(EditAnywhere, Category="Setter")
+	bool Value = false;
+};
+
+/** Sets a boolean value to the bound property reference */
+USTRUCT(DisplayName="Set Boolean", Category="Setter")
+struct FSceneStateSetBoolTask : public FSceneStateTask
+{
+	GENERATED_BODY()
+
+	using FInstanceDataType = FSceneStateSetBoolTaskInstance;
+
+protected:
+	//~ Begin FSceneStateTask
+#if WITH_EDITOR
+	virtual const UScriptStruct* OnGetTaskInstanceType() const override;
+#endif
+	virtual void OnStart(const FSceneStateExecutionContext& InContext, FStructView InTaskInstance) const override;
+	//~ End FSceneStateTask
+};
diff --git a/Engine/Plugins/Experimental/SceneState/Source/SceneStateTasks/Public/Setters/SceneStateSetIntegerTask.h b/Engine/Plugins/Experimental/SceneState/Source/SceneStateTasks/Public/Setters/SceneStateSetIntegerTask.h
new file mode 100644
--- /dev/null
+++ b/Engine/Plugins/Experimental/SceneState/Source/SceneStateTasks/Public/Setters/SceneStateSetIntegerTask.h
@@ -0,0 +1,39 @@
+// Copyright Epic Games, Inc. All Rights Reserved.
+
+#pragma once
+
+#include "SceneStatePropertyReference.h"
+#include "Tasks/SceneStateTask.h"
+#include "Tasks/SceneStateTaskInstance.h"
+#include "SceneStateSetIntegerTask.generated.h"
+
+USTRUCT()
+struct FSceneStateSetIntegerTaskInstance : public FSceneStateTaskInstance
+{
+	GENERATED_BODY()
+
+	/** Target integer property to set */
+	UPROPERTY(EditAnywhere, Category="Setter", meta=(RefType="int32"))
+	FSceneStatePropertyReference Target;
+
+	/** The integer value to set. */
+	UPROPERTY(EditAnywhere, Category="Setter")
+	int32 Value = 0;
+};
+
+/** Sets a integer value to the bound property reference */
+USTRUCT(DisplayName="Set Integer", Category="Setter")
+struct FSceneStateSetIntegerTask : public FSceneStateTask
+{
+	GENERATED_BODY()
+
+	using FInstanceDataType = FSceneStateSetIntegerTaskInstance;
+
+protected:
+	//~ Begin FSceneStateTask
+#if WITH_EDITOR
+	virtual const UScriptStruct* OnGetTaskInstanceType() const override;
+#endif
+	virtual void OnStart(const FSceneStateExecutionContext& InContext, FStructView InTaskInstance) const override;
+	//~ End FSceneStateTask
+};
diff --git a/Engine/Plugins/Experimental/SceneState/Source/SceneStateTasks/Public/Setters/SceneStateSetStringTask.h b/Engine/Plugins/Experimental/SceneState/Source/SceneStateTasks/Public/Setters/SceneStateSetStringTask.h
new file mode 100644
--- /dev/null
+++ b/Engine/Plugins/Experimental/SceneState/Source/SceneStateTasks/Public/Setters/SceneStateSetStringTask.h
@@ -0,0 +1,39 @@
+// Copyright Epic Games, Inc. All Rights Reserved.
+
+#pragma once
+
+#include "SceneStatePropertyReference.h"
+#include "Tasks/SceneStateTask.h"
+#include "Tasks/SceneStateTaskInstance.h"
+#include "SceneStateSetStringTask.generated.h"
+
+USTRUCT()
+struct FSceneStateSetStringTaskInstance : public FSceneStateTaskInstance
+{
+	GENERATED_BODY()
+
+	/** Target string property to set */
+	UPROPERTY(EditAnywhere, Category="Setter", meta=(RefType="String"))
+	FSceneStatePropertyReference Target;
+
+	/** The string value to set. */
+	UPROPERTY(EditAnywhere, Category="Setter")
+	FString Value;
+};
+
+/** Sets a string value to the bound property reference */
+USTRUCT(DisplayName="Set String", Category="Setter")
+struct FSceneStateSetStringTask : public FSceneStateTask
+{
+	GENERATED_BODY()
+
+	using FInstanceDataType = FSceneStateSetStringTaskInstance;
+
+protected:
+	//~ Begin FSceneStateTask
+#if WITH_EDITOR
+	virtual const UScriptStruct* OnGetTaskInstanceType() const override;
+#endif
+	virtual void OnStart(const FSceneStateExecutionContext& InContext, FStructView InTaskInstance) const override;
+	//~ End FSceneStateTask
+};
diff --git a/Engine/Plugins/Experimental/SceneState/Source/SceneStateTasks/Public/Setters/SceneStateSetTextTask.h b/Engine/Plugins/Experimental/SceneState/Source/SceneStateTasks/Public/Setters/SceneStateSetTextTask.h
new file mode 100644
--- /dev/null
+++ b/Engine/Plugins/Experimental/SceneState/Source/SceneStateTasks/Public/Setters/SceneStateSetTextTask.h
@@ -0,0 +1,39 @@
+// Copyright Epic Games, Inc. All Rights Reserved.
+
+#pragma once
+
+#include "SceneStatePropertyReference.h"
+#include "Tasks/SceneStateTask.h"
+#include "Tasks/SceneStateTaskInstance.h"
+#include "SceneStateSetTextTask.generated.h"
+
+USTRUCT()
+struct FSceneStateSetTextTaskInstance : public FSceneStateTaskInstance
+{
+	GENERATED_BODY()
+
+	/** Target text property to set */
+	UPROPERTY(EditAnywhere, Category="Setter", meta=(RefType="Text"))
+	FSceneStatePropertyReference Target;
+
+	/** The text value to set. */
+	UPROPERTY(EditAnywhere, Category="Setter")
+	FText Value;
+};
+
+/** Sets a text value to the bound property reference */
+USTRUCT(DisplayName="Set Text", Category="Setter")
+struct FSceneStateSetTextTask : public FSceneStateTask
+{
+	GENERATED_BODY()
+
+	using FInstanceDataType = FSceneStateSetTextTaskInstance;
+
+protected:
+	//~ Begin FSceneStateTask
+#if WITH_EDITOR
+	virtual const UScriptStruct* OnGetTaskInstanceType() const override;
+#endif
+	virtual void OnStart(const FSceneStateExecutionContext& InContext, FStructView InTaskInstance) const override;
+	//~ End FSceneStateTask
+};
diff --git a/Engine/Plugins/Experimental/SceneState/Source/SceneStateTasks/Public/Setters/SceneStateSetterUtils.h b/Engine/Plugins/Experimental/SceneState/Source/SceneStateTasks/Public/Setters/SceneStateSetterUtils.h
new file mode 100644
--- /dev/null
+++ b/Engine/Plugins/Experimental/SceneState/Source/SceneStateTasks/Public/Setters/SceneStateSetterUtils.h
@@ -0,0 +1,69 @@
+// Copyright Epic Games, Inc. All Rights Reserved.
+
+#pragma once
+
+#include "Concepts/DecaysTo.h"
+#include "PropertyBindingDataView.h"
+#include "SceneStateExecutionContext.h"
+#include "SceneStatePropertyUtils.h"
+
+struct FSceneStatePropertyReference;
+
+namespace UE::SceneState
+{
+	/**
+	 * Attempts to call the Dest Property's setter with a provided source value
+	 * @param InDestValuePtr the value pointer to get the container ptr from (to pass back into FProperty::CallSetter)
+	 * @param InDestProperty the destination property to copy to
+	 * @param InSourceValuePtr pointer to the source data to copy from
+	 * @param InArrayIndex the array index to set within the destination property (necessary if property has an Array Dim > 1, 0 otherwise)
+	 * @return true if the setter call was made
+	 */
+	SCENESTATETASKS_API bool CallValueSetter(uint8* InDestValuePtr, const FProperty* InDestProperty, const uint8* InSourceValuePtr, int32 InArrayIndex = 0);
+
+	/** A Setter Task Instance is valid if it has a valid 'Value' member, as well as a 'Target' FSceneStatePropertyReference member */
+	template<typename InTaskInstanceType>
+	concept CSetterTaskInstanceTypeable = requires(const InTaskInstanceType& InTaskInstance)
+	{
+		InTaskInstance.Value;
+		{ InTaskInstance.Target }->UE::CDecaysTo<FSceneStatePropertyReference>;
+	};
+
+	/** A Setter Task is valid if it has the FInstanceDataType for its Task Instance type, and this Task Instance meets the requirements to be a Setter Task Instance*/
+	template<typename InTaskType>
+	concept CSetterTaskTypeable = CSetterTaskInstanceTypeable<typename InTaskType::FInstanceDataType>;
+
+	/**
+	 * Sets the Task Instance's value to the Task Instance's property reference
+	 * @param InTask the FSceneStateTask object
+	 * @param InContext the current execution context
+	 * @param InTaskInstance view of the task instance data
+	 */
+	template<CSetterTaskTypeable InTaskType>
+	void SetValue(const InTaskType& InTask, const FSceneStateExecutionContext& InContext, FStructView InTaskInstance)
+	{
+		typename InTaskType::FInstanceDataType& Instance = InTaskInstance.Get<typename InTaskType::FInstanceDataType>();
+
+		using FValueType = decltype(Instance.Value);
+
+		const FSceneStateBindingResolvedReference* ResolvedReference = FindResolvedReference(InContext, Instance.Target);
+		if (!ResolvedReference || !ResolvedReference->SourceLeafProperty || !TDataTypeHelper<FValueType>::IsValid(ResolvedReference->SourceLeafProperty))
+		{
+			return;
+		}
+
+		uint8* const ValuePtr = ResolveProperty(InContext, *ResolvedReference);
+		if (!ValuePtr)
+		{
+			return;
+		}
+
+		// Attempt to call setter if available
+		if (!CallValueSetter(ValuePtr, ResolvedReference->SourceLeafProperty, reinterpret_cast<const uint8*>(&Instance.Value)))
+		{
+			// The setter wasn't called, set the value of the property directly
+			*reinterpret_cast<FValueType*>(ValuePtr) = Instance.Value;
+		}
+	}
+
+} // UE::SceneState
diff --git a/Engine/Plugins/Experimental/SceneState/Source/SceneStateTasks/SceneStateTasks.Build.cs b/Engine/Plugins/Experimental/SceneState/Source/SceneStateTasks/SceneStateTasks.Build.cs
--- a/Engine/Plugins/Experimental/SceneState/Source/SceneStateTasks/SceneStateTasks.Build.cs
+++ b/Engine/Plugins/Experimental/SceneState/Source/SceneStateTasks/SceneStateTasks.Build.cs
@@ -15,6 +15,7 @@ public class SceneStateTasks : ModuleRules
                 "CoreUObject",
                 "Engine",
                 "SceneState",
+                "SceneStateBinding",
             }
         );
     }
