Motion Design: scene state property ref to support multiple types and changed the property reference utils template functions for this multiple type support.
Additionally, added float setter task (for double/single precision floats).

diff --git a/Engine/Plugins/Experimental/SceneState/Source/SceneStateBinding/Private/SceneStatePropertyReferenceUtils.cpp b/Engine/Plugins/Experimental/SceneState/Source/SceneStateBinding/Private/SceneStatePropertyReferenceUtils.cpp
--- a/Engine/Plugins/Experimental/SceneState/Source/SceneStateBinding/Private/SceneStatePropertyReferenceUtils.cpp
+++ b/Engine/Plugins/Experimental/SceneState/Source/SceneStateBinding/Private/SceneStatePropertyReferenceUtils.cpp
@@ -29,16 +29,27 @@ static const FLazyName TYPE_Name(TEXT("Name"));
 static const FLazyName TYPE_String(TEXT("String"));
 static const FLazyName TYPE_Text(TEXT("Text"));
 
+TArray<FString> GetPropertyRefTypes(const FProperty* InReferenceProperty)
+{
+	check(IsPropertyReference(InReferenceProperty));
+
+	FString RefType = InReferenceProperty->GetMetaData(Metadata::RefType);
+	RefType.RemoveSpacesInline();
+
+	TArray<FString> TargetTypes;
+	RefType.ParseIntoArray(TargetTypes, TEXT(","), /*bCullEmpty*/true);
+	return TargetTypes;
+}
+	
 bool ArePropertyReferenceCompatible(const FProperty* InTargetReferenceProperty, const FProperty* InSourceReferenceProperty)
 {
 	check(IsPropertyReference(InSourceReferenceProperty) && IsPropertyReference(InTargetReferenceProperty));
 
-	const FEdGraphPinType SourceRefPin = GetPropertyReferencePinType(InSourceReferenceProperty);
-	const FEdGraphPinType TargetRefPin = GetPropertyReferencePinType(InTargetReferenceProperty);
+	const TArray<FEdGraphPinType, TInlineAllocator<1>> SourceRefPins = GetPropertyReferencePinTypes(InSourceReferenceProperty);
+	const TArray<FEdGraphPinType, TInlineAllocator<1>> TargetRefPins = GetPropertyReferencePinTypes(InTargetReferenceProperty);
 
-	return SourceRefPin.PinCategory == TargetRefPin.PinCategory
-		&& SourceRefPin.ContainerType == TargetRefPin.ContainerType 
-		&& SourceRefPin.PinSubCategoryObject == TargetRefPin.PinSubCategoryObject;
+	// NOTE: For now be stringent, and return false for properties that have the same Ref types metadata but in different order.
+	return SourceRefPins == TargetRefPins;
 }
 
 bool IsPropertyCompatible(const FString& InTargetType, const FProperty* InTestProperty, bool bInIsTargetRefArray)
@@ -208,6 +219,20 @@ FEdGraphPinType GetTargetPinType(const FString& InTargetType, EPinContainerType
 
 	return PinType;
 }
+
+/** Retrieves the Ed Graph Pin Type for a given target type (in string) and a desired container type */
+TArray<FEdGraphPinType, TInlineAllocator<1>> GetTargetPinTypes(const TArray<FString>& InTargetTypes, EPinContainerType InContainerType)
+{
+	TArray<FEdGraphPinType, TInlineAllocator<1>> PinTypes;
+	PinTypes.Reserve(InTargetTypes.Num());
+
+	for (const FString& TargetType : InTargetTypes)
+	{
+		PinTypes.Add(GetTargetPinType(TargetType, InContainerType));
+	}
+
+	return PinTypes;
+}
 #endif
 
 } // UE::SceneState::Private
@@ -250,15 +275,30 @@ bool IsPropertyReferenceCompatible(const FProperty* InReferenceProperty, const F
 		}
 	}
 
-	const FString TargetType = InReferenceProperty->GetMetaData(Metadata::RefType);
-	return Private::IsPropertyCompatible(TargetType, TestProperty, bIsTargetRefArray);
+	const TArray<FString> TargetTypes = Private::GetPropertyRefTypes(InReferenceProperty);
+	for (const FString& TargetType : TargetTypes)
+	{
+		if (Private::IsPropertyCompatible(TargetType, TestProperty, bIsTargetRefArray))
+		{
+			return true;
+		}
+	}
+	return false;
 }
 
 FEdGraphPinType GetPropertyReferencePinType(const FProperty* InProperty)
+{
+	const TArray<FEdGraphPinType, TInlineAllocator<1>> PinTypes = GetPropertyReferencePinTypes(InProperty);
+	return !PinTypes.IsEmpty()
+		? PinTypes[0]
+		: FEdGraphPinType();
+}
+
+TArray<FEdGraphPinType, TInlineAllocator<1>> GetPropertyReferencePinTypes(const FProperty* InProperty)
 {
 	if (!ensure(InProperty))
 	{
-		return FEdGraphPinType();
+		return {};
 	}
 
 	ensure(IsPropertyReference(InProperty));
@@ -267,8 +307,7 @@ FEdGraphPinType GetPropertyReferencePinType(const FProperty* InProperty)
 		? EPinContainerType::Array
 		: EPinContainerType::None;
 
-	const FString TargetType = InProperty->GetMetaData(Metadata::RefType);
-	return Private::GetTargetPinType(TargetType, ContainerType);
+	return Private::GetTargetPinTypes(Private::GetPropertyRefTypes(InProperty), ContainerType);
 }
 #endif
 
diff --git a/Engine/Plugins/Experimental/SceneState/Source/SceneStateBinding/Public/SceneStatePropertyReferenceUtils.h b/Engine/Plugins/Experimental/SceneState/Source/SceneStateBinding/Public/SceneStatePropertyReferenceUtils.h
--- a/Engine/Plugins/Experimental/SceneState/Source/SceneStateBinding/Public/SceneStatePropertyReferenceUtils.h
+++ b/Engine/Plugins/Experimental/SceneState/Source/SceneStateBinding/Public/SceneStatePropertyReferenceUtils.h
@@ -18,8 +18,11 @@ namespace UE::SceneState
 	/** Returns whether a source property is compatible with the reference property */
 	SCENESTATEBINDING_API bool IsPropertyReferenceCompatible(const FProperty* InReferenceProperty, const FProperty* InSourceProperty);
 
-	/** Returns the graph pin types based on a given Property Reference property's meta-data */
+	/** Returns the first graph pin type based on a given Property Reference property's meta-data */
 	SCENESTATEBINDING_API FEdGraphPinType GetPropertyReferencePinType(const FProperty* InProperty);
+
+	/** Returns the graph pin types based on a given Property Reference property's meta-data */
+	SCENESTATEBINDING_API TArray<FEdGraphPinType, TInlineAllocator<1>> GetPropertyReferencePinTypes(const FProperty* InProperty);
 #endif
 
 	template<class InDataType>
diff --git a/Engine/Plugins/Experimental/SceneState/Source/SceneStateBlueprintEditor/Private/DetailsView/SceneStateCachedBindingData.cpp b/Engine/Plugins/Experimental/SceneState/Source/SceneStateBlueprintEditor/Private/DetailsView/SceneStateCachedBindingData.cpp
--- a/Engine/Plugins/Experimental/SceneState/Source/SceneStateBlueprintEditor/Private/DetailsView/SceneStateCachedBindingData.cpp
+++ b/Engine/Plugins/Experimental/SceneState/Source/SceneStateBlueprintEditor/Private/DetailsView/SceneStateCachedBindingData.cpp
@@ -1,6 +1,7 @@
 // Copyright Epic Games, Inc. All Rights Reserved.
 
 #include "SceneStateCachedBindingData.h"
+#include "EdGraphSchema_K2.h"
 #include "IPropertyAccessEditor.h"
 #include "PropertyBindingDataView.h"
 #include "SceneStateBindingDesc.h"
@@ -59,29 +60,37 @@ void FSceneStateCachedBindingData::AddPropertyInfoOverride(const FProperty& InPr
 	// Add the PropertyRef property type with its RefType
 	if (StructProperty->Struct->IsChildOf<FSceneStatePropertyReference>())
 	{
-		FEdGraphPinType PinType = GetPropertyReferencePinType(&InProperty);
+		TArray<FEdGraphPinType, TInlineAllocator<1>> PinTypes = GetPropertyReferencePinTypes(&InProperty);
 
 		// If Property supports Arrays, add the Array version of these pin types
-		if (GetPropertyHandle()->HasMetaData(Metadata::CanRefToArray))
-		{
-			PinType.ContainerType = EPinContainerType::Array;
-		}
+		const bool bCanRefToArray = GetPropertyHandle()->HasMetaData(Metadata::CanRefToArray);
 
-		TSharedRef<UE::PropertyBinding::FPropertyInfoOverride> PropertyInfoOverride = MakeShared<UE::PropertyBinding::FPropertyInfoOverride>();
-		PropertyInfoOverride->PinType = PinType;
+		OutPropertyInfoOverrides.Reserve(OutPropertyInfoOverrides.Num() + PinTypes.Num());
 
-		FString TypeName;
-		if (UObject* SubCategoryObject = PinType.PinSubCategoryObject.Get())
+		for (FEdGraphPinType& PinType : PinTypes)
 		{
-			TypeName = SubCategoryObject->GetName();
+			if (bCanRefToArray)
+			{
+				PinType.ContainerType = EPinContainerType::Array;
+			}
+
+
+			TSharedRef<UE::PropertyBinding::FPropertyInfoOverride> PropertyInfoOverride = MakeShared<UE::PropertyBinding::FPropertyInfoOverride>();
+			PropertyInfoOverride->PinType = PinType;
+
+			FString TypeName;
+			if (UObject* SubCategoryObject = PinType.PinSubCategoryObject.Get())
+			{
+				TypeName = SubCategoryObject->GetName();
+			}
+			else
+			{
+				TypeName = PinType.PinCategory.ToString() + TEXT(" ") + PinType.PinSubCategory.ToString();
+			}
+
+			PropertyInfoOverride->TypeNameText = FText::FromString(TypeName);
+			OutPropertyInfoOverrides.Emplace(MoveTemp(PropertyInfoOverride));
 		}
-		else
-		{
-			TypeName = PinType.PinCategory.ToString() + TEXT(" ") + PinType.PinSubCategory.ToString();
-		}
-
-		PropertyInfoOverride->TypeNameText = FText::FromString(TypeName);
-		OutPropertyInfoOverrides.Emplace(MoveTemp(PropertyInfoOverride));
 	}
 }
 
diff --git a/Engine/Plugins/Experimental/SceneState/Source/SceneStateTasks/Private/Setters/SceneStateSetBoolTask.cpp b/Engine/Plugins/Experimental/SceneState/Source/SceneStateTasks/Private/Setters/SceneStateSetBoolTask.cpp
--- a/Engine/Plugins/Experimental/SceneState/Source/SceneStateTasks/Private/Setters/SceneStateSetBoolTask.cpp
+++ b/Engine/Plugins/Experimental/SceneState/Source/SceneStateTasks/Private/Setters/SceneStateSetBoolTask.cpp
@@ -12,6 +12,6 @@ const UScriptStruct* FSceneStateSetBoolTask::OnGetTaskInstanceType() const
 
 void FSceneStateSetBoolTask::OnStart(const FSceneStateExecutionContext& InContext, FStructView InTaskInstance) const
 {
-	UE::SceneState::SetValue(*this, InContext, InTaskInstance);
+	UE::SceneState::SetValue<bool>(*this, InContext, InTaskInstance);
 	Finish(InContext, InTaskInstance);
 }
diff --git a/Engine/Plugins/Experimental/SceneState/Source/SceneStateTasks/Private/Setters/SceneStateSetFloatTask.cpp b/Engine/Plugins/Experimental/SceneState/Source/SceneStateTasks/Private/Setters/SceneStateSetFloatTask.cpp
new file mode 100644
--- /dev/null
+++ b/Engine/Plugins/Experimental/SceneState/Source/SceneStateTasks/Private/Setters/SceneStateSetFloatTask.cpp
@@ -0,0 +1,17 @@
+// Copyright Epic Games, Inc. All Rights Reserved.
+
+#include "Setters/SceneStateSetFloatTask.h"
+#include "Setters/SceneStateSetterUtils.h"
+
+#if WITH_EDITOR
+const UScriptStruct* FSceneStateSetFloatTask::OnGetTaskInstanceType() const
+{
+	return FInstanceDataType::StaticStruct();
+}
+#endif
+
+void FSceneStateSetFloatTask::OnStart(const FSceneStateExecutionContext& InContext, FStructView InTaskInstance) const
+{
+	UE::SceneState::SetValue<double, float>(*this, InContext, InTaskInstance);
+	Finish(InContext, InTaskInstance);
+}
diff --git a/Engine/Plugins/Experimental/SceneState/Source/SceneStateTasks/Private/Setters/SceneStateSetIntegerTask.cpp b/Engine/Plugins/Experimental/SceneState/Source/SceneStateTasks/Private/Setters/SceneStateSetIntegerTask.cpp
--- a/Engine/Plugins/Experimental/SceneState/Source/SceneStateTasks/Private/Setters/SceneStateSetIntegerTask.cpp
+++ b/Engine/Plugins/Experimental/SceneState/Source/SceneStateTasks/Private/Setters/SceneStateSetIntegerTask.cpp
@@ -12,6 +12,6 @@ const UScriptStruct* FSceneStateSetIntegerTask::OnGetTaskInstanceType() const
 
 void FSceneStateSetIntegerTask::OnStart(const FSceneStateExecutionContext& InContext, FStructView InTaskInstance) const
 {
-	UE::SceneState::SetValue(*this, InContext, InTaskInstance);
+	UE::SceneState::SetValue<int32>(*this, InContext, InTaskInstance);
 	Finish(InContext, InTaskInstance);
 }
diff --git a/Engine/Plugins/Experimental/SceneState/Source/SceneStateTasks/Private/Setters/SceneStateSetStringTask.cpp b/Engine/Plugins/Experimental/SceneState/Source/SceneStateTasks/Private/Setters/SceneStateSetStringTask.cpp
--- a/Engine/Plugins/Experimental/SceneState/Source/SceneStateTasks/Private/Setters/SceneStateSetStringTask.cpp
+++ b/Engine/Plugins/Experimental/SceneState/Source/SceneStateTasks/Private/Setters/SceneStateSetStringTask.cpp
@@ -12,6 +12,6 @@ const UScriptStruct* FSceneStateSetStringTask::OnGetTaskInstanceType() const
 
 void FSceneStateSetStringTask::OnStart(const FSceneStateExecutionContext& InContext, FStructView InTaskInstance) const
 {
-	UE::SceneState::SetValue(*this, InContext, InTaskInstance);
+	UE::SceneState::SetValue<FString>(*this, InContext, InTaskInstance);
 	Finish(InContext, InTaskInstance);
 }
diff --git a/Engine/Plugins/Experimental/SceneState/Source/SceneStateTasks/Private/Setters/SceneStateSetTextTask.cpp b/Engine/Plugins/Experimental/SceneState/Source/SceneStateTasks/Private/Setters/SceneStateSetTextTask.cpp
--- a/Engine/Plugins/Experimental/SceneState/Source/SceneStateTasks/Private/Setters/SceneStateSetTextTask.cpp
+++ b/Engine/Plugins/Experimental/SceneState/Source/SceneStateTasks/Private/Setters/SceneStateSetTextTask.cpp
@@ -12,6 +12,6 @@ const UScriptStruct* FSceneStateSetTextTask::OnGetTaskInstanceType() const
 
 void FSceneStateSetTextTask::OnStart(const FSceneStateExecutionContext& InContext, FStructView InTaskInstance) const
 {
-	UE::SceneState::SetValue(*this, InContext, InTaskInstance);
+	UE::SceneState::SetValue<FText>(*this, InContext, InTaskInstance);
 	Finish(InContext, InTaskInstance);
 }
diff --git a/Engine/Plugins/Experimental/SceneState/Source/SceneStateTasks/Private/Setters/SceneStateSetterUtils.cpp b/Engine/Plugins/Experimental/SceneState/Source/SceneStateTasks/Private/Setters/SceneStateSetterUtils.cpp
--- a/Engine/Plugins/Experimental/SceneState/Source/SceneStateTasks/Private/Setters/SceneStateSetterUtils.cpp
+++ b/Engine/Plugins/Experimental/SceneState/Source/SceneStateTasks/Private/Setters/SceneStateSetterUtils.cpp
@@ -5,26 +5,27 @@
 namespace UE::SceneState
 {
 
-bool CallValueSetter(uint8* InDestValuePtr, const FProperty* InDestProperty, const uint8* InSourceValuePtr, int32 InArrayIndex)
+bool CallValueSetter(uint8* InDestValuePtr, const FSceneStateBindingResolvedReference& InPropertyReference, const uint8* InSourceValuePtr, int32 InArrayIndex)
 {
-	if (!ensureMsgf(InArrayIndex >= 0 && InArrayIndex < InDestProperty->ArrayDim, TEXT("Index %d outside property array dim bounds %d!"), InArrayIndex, InDestProperty->ArrayDim))
+	const FProperty* const DestProperty = InPropertyReference.SourceLeafProperty;
+	if (!ensureMsgf(InArrayIndex >= 0 && InArrayIndex < DestProperty->ArrayDim, TEXT("Index %d outside property array dim bounds %d!"), InArrayIndex, DestProperty->ArrayDim))
 	{
 		return false;
 	}
 
 	// No setter, can't call via this method
-	if (!InDestProperty->HasSetter())
+	if (!DestProperty->HasSetter())
 	{
 		return false;
 	}
 
 	// Reverse logic of ContainerPtrToValuePtr (i.e. ValuePtrToContainerPtr)
-	uint8* const DestContainerPtr = InDestValuePtr - (InDestProperty->GetOffset_ForInternal() + InDestProperty->GetElementSize() * InArrayIndex);
+	uint8* const DestContainerPtr = InDestValuePtr - (DestProperty->GetOffset_ForInternal() + DestProperty->GetElementSize() * InArrayIndex);
 
 	// check that this reverse logic remains 'up-to-date' with the latest ContainerPtrToValuePtr logic
-	checkf(InDestProperty->ContainerPtrToValuePtr<uint8>(DestContainerPtr) == InDestValuePtr, TEXT("ValuePtrToContainerPtr logic outdated. Verify against ContainerPtrToValuePtr."));
+	checkf(DestProperty->ContainerPtrToValuePtr<uint8>(DestContainerPtr) == InDestValuePtr, TEXT("ValuePtrToContainerPtr logic outdated. Verify against ContainerPtrToValuePtr."));
 
-	InDestProperty->CallSetter(DestContainerPtr, InSourceValuePtr);
+	DestProperty->CallSetter(DestContainerPtr, InSourceValuePtr);
 	return true;
 }
 
diff --git a/Engine/Plugins/Experimental/SceneState/Source/SceneStateTasks/Public/Setters/SceneStateSetFloatTask.h b/Engine/Plugins/Experimental/SceneState/Source/SceneStateTasks/Public/Setters/SceneStateSetFloatTask.h
new file mode 100644
--- /dev/null
+++ b/Engine/Plugins/Experimental/SceneState/Source/SceneStateTasks/Public/Setters/SceneStateSetFloatTask.h
@@ -0,0 +1,39 @@
+// Copyright Epic Games, Inc. All Rights Reserved.
+
+#pragma once
+
+#include "SceneStatePropertyReference.h"
+#include "Tasks/SceneStateTask.h"
+#include "Tasks/SceneStateTaskInstance.h"
+#include "SceneStateSetFloatTask.generated.h"
+
+USTRUCT()
+struct FSceneStateSetFloatTaskInstance : public FSceneStateTaskInstance
+{
+	GENERATED_BODY()
+
+	/** Target float property to set */
+	UPROPERTY(EditAnywhere, Category="Setter", meta=(RefType="double,float"))
+	FSceneStatePropertyReference Target;
+
+	/** The float value to set */
+	UPROPERTY(EditAnywhere, Category="Setter")
+	double Value = 0.0;
+};
+
+/** Sets a float value to the bound property reference */
+USTRUCT(DisplayName="Set Float", Category="Setter")
+struct FSceneStateSetFloatTask : public FSceneStateTask
+{
+	GENERATED_BODY()
+
+	using FInstanceDataType = FSceneStateSetFloatTaskInstance;
+
+protected:
+	//~ Begin FSceneStateTask
+#if WITH_EDITOR
+	virtual const UScriptStruct* OnGetTaskInstanceType() const override;
+#endif
+	virtual void OnStart(const FSceneStateExecutionContext& InContext, FStructView InTaskInstance) const override;
+	//~ End FSceneStateTask
+};
diff --git a/Engine/Plugins/Experimental/SceneState/Source/SceneStateTasks/Public/Setters/SceneStateSetterUtils.h b/Engine/Plugins/Experimental/SceneState/Source/SceneStateTasks/Public/Setters/SceneStateSetterUtils.h
--- a/Engine/Plugins/Experimental/SceneState/Source/SceneStateTasks/Public/Setters/SceneStateSetterUtils.h
+++ b/Engine/Plugins/Experimental/SceneState/Source/SceneStateTasks/Public/Setters/SceneStateSetterUtils.h
@@ -14,12 +14,12 @@ namespace UE::SceneState
 	/**
 	 * Attempts to call the Dest Property's setter with a provided source value
 	 * @param InDestValuePtr the value pointer to get the container ptr from (to pass back into FProperty::CallSetter)
-	 * @param InDestProperty the destination property to copy to
+	 * @param InPropertyReference the property reference describing the destination value ptr
 	 * @param InSourceValuePtr pointer to the source data to copy from
 	 * @param InArrayIndex the array index to set within the destination property (necessary if property has an Array Dim > 1, 0 otherwise)
 	 * @return true if the setter call was made
 	 */
-	SCENESTATETASKS_API bool CallValueSetter(uint8* InDestValuePtr, const FProperty* InDestProperty, const uint8* InSourceValuePtr, int32 InArrayIndex = 0);
+	SCENESTATETASKS_API bool CallValueSetter(uint8* InDestValuePtr, const FSceneStateBindingResolvedReference& InPropertyReference, const uint8* InSourceValuePtr, int32 InArrayIndex = 0);
 
 	/** A Setter Task Instance is valid if it has a valid 'Value' member, as well as a 'Target' FSceneStatePropertyReference member */
 	template<typename InTaskInstanceType>
@@ -33,37 +33,52 @@ namespace UE::SceneState
 	template<typename InTaskType>
 	concept CSetterTaskTypeable = CSetterTaskInstanceTypeable<typename InTaskType::FInstanceDataType>;
 
+	/**
+	 * Copies the provided value to the given destination address, calling a Setter if available
+	 * @param InDestValuePtr pointer of the destination address to set
+	 * @param InPropertyReference the property reference describing the destination value ptr
+	 * @param InValue the value to copy from
+	 */
+	template<typename InTargetType, typename InValueType>
+	bool SetValue(uint8* InDestValuePtr, const FSceneStateBindingResolvedReference& InPropertyReference, const InValueType& InValue)
+	{
+		if (TDataTypeHelper<InTargetType>::IsValid(InPropertyReference.SourceLeafProperty))
+		{
+			// Attempt to call setter if available
+			if (!CallValueSetter(InDestValuePtr, InPropertyReference, reinterpret_cast<const uint8*>(&InValue)))
+			{
+				// The setter wasn't called, set the value of the property directly
+				*reinterpret_cast<InTargetType*>(InDestValuePtr) = InValue;
+			}
+			return true;
+		}
+		return false;
+	}
+
 	/**
 	 * Sets the Task Instance's value to the Task Instance's property reference
 	 * @param InTask the FSceneStateTask object
 	 * @param InContext the current execution context
 	 * @param InTaskInstance view of the task instance data
 	 */
-	template<CSetterTaskTypeable InTaskType>
-	void SetValue(const InTaskType& InTask, const FSceneStateExecutionContext& InContext, FStructView InTaskInstance)
+	template<typename... InTargetTypes, CSetterTaskTypeable InTaskType>
+	bool SetValue(const InTaskType& InTask, const FSceneStateExecutionContext& InContext, FStructView InTaskInstance)
 	{
 		typename InTaskType::FInstanceDataType& Instance = InTaskInstance.Get<typename InTaskType::FInstanceDataType>();
 
-		using FValueType = decltype(Instance.Value);
-
 		const FSceneStateBindingResolvedReference* ResolvedReference = FindResolvedReference(InContext, Instance.Target);
-		if (!ResolvedReference || !ResolvedReference->SourceLeafProperty || !TDataTypeHelper<FValueType>::IsValid(ResolvedReference->SourceLeafProperty))
+		if (!ResolvedReference || !ResolvedReference->SourceLeafProperty)
 		{
-			return;
+			return false;
 		}
 
 		uint8* const ValuePtr = ResolveProperty(InContext, *ResolvedReference);
 		if (!ValuePtr)
 		{
-			return;
+			return false;
 		}
 
-		// Attempt to call setter if available
-		if (!CallValueSetter(ValuePtr, ResolvedReference->SourceLeafProperty, reinterpret_cast<const uint8*>(&Instance.Value)))
-		{
-			// The setter wasn't called, set the value of the property directly
-			*reinterpret_cast<FValueType*>(ValuePtr) = Instance.Value;
-		}
+		return (SetValue<InTargetTypes>(ValuePtr, *ResolvedReference, Instance.Value) || ...);
 	}
 
 } // UE::SceneState
