Motion Design: reduced the size of base Scene State task instance by removing elapsed time and result (unused), and moving the instance id to the state (since a task instance is bound to a state instance). Elapsed time should now be done per implementation (e.g. Delay now has the elapsed time entry).

diff --git a/Engine/Plugins/Experimental/SceneState/Source/SceneState/Private/SceneState.cpp b/Engine/Plugins/Experimental/SceneState/Source/SceneState/Private/SceneState.cpp
--- a/Engine/Plugins/Experimental/SceneState/Source/SceneState/Private/SceneState.cpp
+++ b/Engine/Plugins/Experimental/SceneState/Source/SceneState/Private/SceneState.cpp
@@ -173,7 +173,7 @@ void FSceneState::UpdateActiveTasks(const FSceneStateExecutionContext& InContext
 		[&InContext, &InInstance, &FindTaskInstance](const FSceneStateTask& InTask, FStructView InTaskInstance)->EIterationResult
 		{
 			FSceneStateTaskInstance* TaskInstance = InTaskInstance.GetPtr<FSceneStateTaskInstance>();
-			if (!TaskInstance || TaskInstance->Status != EExecutionStatus::NotStarted)
+			if (!TaskInstance || TaskInstance->GetStatus() != EExecutionStatus::NotStarted)
 			{
 				return EIterationResult::Continue;
 			}
@@ -182,7 +182,7 @@ void FSceneState::UpdateActiveTasks(const FSceneStateExecutionContext& InContext
 			for (uint16 Prerequisite : InContext.GetTaskPrerequisites(InTask))
 			{
 				FSceneStateTaskInstance* const PrerequisiteTaskInstance = FindTaskInstance(Prerequisite);
-				if (!PrerequisiteTaskInstance || PrerequisiteTaskInstance->Status != EExecutionStatus::Finished)
+				if (!PrerequisiteTaskInstance || PrerequisiteTaskInstance->GetStatus() != EExecutionStatus::Finished)
 				{
 					bPrerequisitesMet = false;
 					break;
@@ -225,7 +225,7 @@ bool FSceneState::HasPendingTasks(const FSceneStateExecutionContext& InContext)
 		[StateInstance, &bHasPendingTask](const FSceneStateTask& InTask, FStructView InTaskInstance)->EIterationResult
 		{
 			FSceneStateTaskInstance* TaskInstance = InTaskInstance.GetPtr<FSceneStateTaskInstance>();
-			if (TaskInstance && TaskInstance->Status != EExecutionStatus::Finished)
+			if (TaskInstance && TaskInstance->GetStatus() != EExecutionStatus::Finished)
 			{
 				bHasPendingTask = true;
 				return EIterationResult::Break;
diff --git a/Engine/Plugins/Experimental/SceneState/Source/SceneState/Private/SceneStateInstance.cpp b/Engine/Plugins/Experimental/SceneState/Source/SceneState/Private/SceneStateInstance.cpp
new file mode 100644
--- /dev/null
+++ b/Engine/Plugins/Experimental/SceneState/Source/SceneState/Private/SceneStateInstance.cpp
@@ -0,0 +1,30 @@
+// Copyright Epic Games, Inc. All Rights Reserved.
+
+#include "SceneStateInstance.h"
+
+namespace UE::SceneState::Private
+{
+	uint16 GetNextInstanceId()
+	{
+		static uint16 Id = 0;
+
+		// Starts with 1
+		uint16 NewId = ++Id;
+		if (NewId == 0)
+		{
+			NewId = ++Id;
+		}
+		return NewId; 
+	}
+
+} // UE::SceneState::Private
+
+FSceneStateInstance::FSceneStateInstance()
+	: InstanceId(UE::SceneState::Private::GetNextInstanceId())
+{
+}
+
+uint16 FSceneStateInstance::GetInstanceId() const
+{
+	return InstanceId;
+}
diff --git a/Engine/Plugins/Experimental/SceneState/Source/SceneState/Private/Tasks/SceneStateBlueprintableTask.cpp b/Engine/Plugins/Experimental/SceneState/Source/SceneState/Private/Tasks/SceneStateBlueprintableTask.cpp
--- a/Engine/Plugins/Experimental/SceneState/Source/SceneState/Private/Tasks/SceneStateBlueprintableTask.cpp
+++ b/Engine/Plugins/Experimental/SceneState/Source/SceneState/Private/Tasks/SceneStateBlueprintableTask.cpp
@@ -6,6 +6,7 @@
 #include "SceneStateExecutionContext.h"
 #include "SceneStateLog.h"
 #include "SceneStateObject.h"
+#include "Tasks/SceneStateBlueprintableTaskWrapper.h"
 #include "Tasks/SceneStateTask.h"
 #include "Tasks/SceneStateTaskInstance.h"
 
@@ -41,23 +42,11 @@ USceneStateEventStream* USceneStateBlueprintableTask::GetEventStream() const
 	return nullptr;
 }
 
-void USceneStateBlueprintableTask::FinishTask(ESceneStateTaskResult InTaskResult)
+void USceneStateBlueprintableTask::FinishTask()
 {
 	TaskExecutionContext.FinishTask();
 }
 
-float USceneStateBlueprintableTask::GetTaskElapsedTime() const
-{
-	FSceneStateTaskInstance* TaskInstance = TaskExecutionContext.GetTaskInstance().GetPtr<FSceneStateTaskInstance>();
-	return TaskInstance ? TaskInstance->ElapsedTime : 0.f;
-}
-
-ESceneStateTaskResult USceneStateBlueprintableTask::GetTaskResult() const
-{
-	FSceneStateTaskInstance* TaskInstance = TaskExecutionContext.GetTaskInstance().GetPtr<FSceneStateTaskInstance>();
-	return TaskInstance ? TaskInstance->Result : ESceneStateTaskResult::Undetermined;
-}
-
 bool USceneStateBlueprintableTask::FindEventHandlerId(const FSceneStateEventSchemaHandle& InEventSchemaHandle, FGuid& OutHandlerId) const
 {
 	const FSceneStateTask* Task = TaskExecutionContext.GetTask().GetPtr<const FSceneStateTask>();
diff --git a/Engine/Plugins/Experimental/SceneState/Source/SceneState/Private/Tasks/SceneStateTask.cpp b/Engine/Plugins/Experimental/SceneState/Source/SceneState/Private/Tasks/SceneStateTask.cpp
--- a/Engine/Plugins/Experimental/SceneState/Source/SceneState/Private/Tasks/SceneStateTask.cpp
+++ b/Engine/Plugins/Experimental/SceneState/Source/SceneState/Private/Tasks/SceneStateTask.cpp
@@ -69,7 +69,7 @@ void FSceneStateTask::Setup(const FSceneStateExecutionContext& InContext, FStruc
 		return;
 	}
 
-	Instance->Status = EExecutionStatus::NotStarted;
+	Instance->SetStatus(EExecutionStatus::NotStarted);
 	OnSetup(InContext, InTaskInstance);
 }
 
@@ -80,14 +80,12 @@ void FSceneStateTask::Start(const FSceneStateExecutionContext& InContext, FStruc
 	QUICK_SCOPE_CYCLE_COUNTER(SceneState_OnStart);
 
 	FSceneStateTaskInstance* Instance = InTaskInstance.GetPtr<FSceneStateTaskInstance>();
-	if (!Instance || Instance->Status != EExecutionStatus::NotStarted)
+	if (!Instance || Instance->GetStatus() != EExecutionStatus::NotStarted)
 	{
 		return;
 	}
 
-	Instance->Status = EExecutionStatus::Running;
-	Instance->Result = ESceneStateTaskResult::Undetermined;
-	Instance->ElapsedTime = 0.f;
+	Instance->SetStatus(EExecutionStatus::Running);
 
 	ApplyBindings(InContext, InTaskInstance);
 
@@ -98,21 +96,21 @@ void FSceneStateTask::Tick(const FSceneStateExecutionContext& InContext, FStruct
 {
 	using namespace UE::SceneState;
 
-	QUICK_SCOPE_CYCLE_COUNTER(SceneState_OnTick);
-
-	FSceneStateTaskInstance* Instance = InTaskInstance.GetPtr<FSceneStateTaskInstance>();
-	if (!Instance || Instance->Status != EExecutionStatus::Running)
+	// Skip task is not set to tick
+	if (!EnumHasAnyFlags(TaskFlags, ESceneStateTaskFlags::Ticks))
 	{
 		return;
 	}
 
-	// Even if Task doesn't tick the virtual/bp functions, keep elapsed time tracked
-	Instance->ElapsedTime += InDeltaSeconds;
+	QUICK_SCOPE_CYCLE_COUNTER(SceneState_OnTick);
 
-	if (EnumHasAllFlags(TaskFlags, ESceneStateTaskFlags::Ticks))
+	FSceneStateTaskInstance* Instance = InTaskInstance.GetPtr<FSceneStateTaskInstance>();
+	if (!Instance || Instance->GetStatus() != EExecutionStatus::Running)
 	{
-		OnTick(InContext, InTaskInstance, InDeltaSeconds);
+		return;
 	}
+
+	OnTick(InContext, InTaskInstance, InDeltaSeconds);
 }
 
 void FSceneStateTask::Stop(const FSceneStateExecutionContext& InContext, FStructView InTaskInstance, ESceneStateTaskStopReason InStopReason) const
@@ -127,13 +125,12 @@ void FSceneStateTask::Stop(const FSceneStateExecutionContext& InContext, FStruct
 		return;
 	}
 
-	if (Instance->Status == EExecutionStatus::Running)
+	if (Instance->GetStatus() == EExecutionStatus::Running)
 	{
 		OnStop(InContext, InTaskInstance, InStopReason);
 	}
 
-	Instance->ElapsedTime = 0.f;
-	Instance->Status = Private::GetTaskStatus(InStopReason);
+	Instance->SetStatus(Private::GetTaskStatus(InStopReason));
 }
 
 void FSceneStateTask::Finish(const FSceneStateExecutionContext& InContext, FStructView InTaskInstance) const
diff --git a/Engine/Plugins/Experimental/SceneState/Source/SceneState/Private/Tasks/SceneStateTaskExecutionContext.cpp b/Engine/Plugins/Experimental/SceneState/Source/SceneState/Private/Tasks/SceneStateTaskExecutionContext.cpp
--- a/Engine/Plugins/Experimental/SceneState/Source/SceneState/Private/Tasks/SceneStateTaskExecutionContext.cpp
+++ b/Engine/Plugins/Experimental/SceneState/Source/SceneState/Private/Tasks/SceneStateTaskExecutionContext.cpp
@@ -22,20 +22,20 @@ FExecutionContextHandle GetContextHandle(const FSceneStateExecutionContext& InCo
 	return FExecutionContextHandle();
 }
 
-uint32 GetTaskInstanceId(const FSceneStateTask& InTask, const FSceneStateExecutionContext& InContext)
+uint16 GetStateInstanceId(const FSceneStateTask& InTask, const FSceneStateExecutionContext& InContext)
 {
-	if (FSceneStateTaskInstance* TaskInstance = InContext.FindTaskInstance(InTask.GetTaskIndex()).GetPtr<FSceneStateTaskInstance>())
+	if (FSceneStateInstance* StateInstance = InContext.FindStateInstance(InTask.GetParentStateIndex()))
 	{
-		return TaskInstance->GetInstanceId();
+		return StateInstance->GetInstanceId();
 	}
 	return 0;
 }
 
-bool IsMatchingInstanceId(FConstStructView InTaskInstanceView, uint32 InTaskInstanceId)
+bool IsMatchingInstanceId(const FSceneStateExecutionContext& InContext, uint16 InParentStateIndex, uint16 InStateInstanceId)
 {
-	if (const FSceneStateTaskInstance* TaskInstance = InTaskInstanceView.GetPtr<const FSceneStateTaskInstance>())
+	if (const FSceneStateInstance* StateInstance = InContext.FindStateInstance(InParentStateIndex))
 	{
-		return TaskInstance->GetInstanceId() == InTaskInstanceId;
+		return StateInstance->GetInstanceId() == InStateInstanceId;
 	}
 	return false;
 }
@@ -44,7 +44,8 @@ bool IsMatchingInstanceId(FConstStructView InTaskInstanceView, uint32 InTaskInst
 
 FTaskExecutionContext::FTaskExecutionContext()
 	: TaskIndex(FSceneStateRange::InvalidIndex)
-	, TaskInstanceId(0)
+	, StateIndex(FSceneStateRange::InvalidIndex)
+	, StateInstanceId(0)
 {
 }
 
@@ -52,7 +53,8 @@ FTaskExecutionContext::FTaskExecutionContext(const FSceneStateTask& InTask, cons
 	: ContextHandle(Private::GetContextHandle(InContext))
 	, ContextRegistryWeak(InContext.GetContextRegistry())
 	, TaskIndex(InTask.GetTaskIndex())
-	, TaskInstanceId(Private::GetTaskInstanceId(InTask, InContext))
+	, StateIndex(InTask.GetParentStateIndex())
+	, StateInstanceId(Private::GetStateInstanceId(InTask, InContext))
 {
 }
 
@@ -69,8 +71,8 @@ FConstStructView FTaskExecutionContext::GetTask() const
 {
 	if (const FSceneStateExecutionContext* ExecutionContext = GetExecutionContext())
 	{
-		// Ensure this Context is still valid
-		if (Private::IsMatchingInstanceId(ExecutionContext->FindTaskInstance(TaskIndex), TaskInstanceId))
+		// Ensure this Context is still valid, by ensuring the saved State Instance Id matches the current State Instance id of execution
+		if (Private::IsMatchingInstanceId(*ExecutionContext, StateIndex, StateInstanceId))
 		{
 			return ExecutionContext->FindTask(TaskIndex);
 		}
@@ -84,8 +86,8 @@ FStructView FTaskExecutionContext::GetTaskInstance() const
 	{
 		FStructView TaskInstance = ExecutionContext->FindTaskInstance(TaskIndex);
 
-		// Ensure this Context is still valid
-		if (Private::IsMatchingInstanceId(TaskInstance, TaskInstanceId))
+		// Ensure this Context is still valid, by ensuring the saved State Instance Id matches the current State Instance id of execution
+		if (Private::IsMatchingInstanceId(*ExecutionContext, StateIndex, StateInstanceId))
 		{
 			return TaskInstance;
 		}
@@ -126,7 +128,7 @@ void FTaskExecutionContext::FinishTask() const
 		return;
 	}
 
-	if (!Private::IsMatchingInstanceId(TaskInstance, TaskInstanceId))
+	if (!Private::IsMatchingInstanceId(*ExecutionContext, StateIndex, StateInstanceId))
 	{
 		UE_LOG(LogSceneState, Error, TEXT("[%s in Context '%s'] %s Task Context outdated!")
 			, *GetDebugString()
@@ -140,7 +142,7 @@ void FTaskExecutionContext::FinishTask() const
 
 FString FTaskExecutionContext::GetDebugString() const
 {
-	return FString::Printf(TEXT("Task Index: %d, State Instance Id: %u"), TaskIndex, TaskInstanceId);
+	return FString::Printf(TEXT("Task Index: %d, State Instance Id: %u"), TaskIndex, StateInstanceId);
 }
 
 } // UE::SceneState
diff --git a/Engine/Plugins/Experimental/SceneState/Source/SceneState/Private/Tasks/SceneStateTaskInstance.cpp b/Engine/Plugins/Experimental/SceneState/Source/SceneState/Private/Tasks/SceneStateTaskInstance.cpp
--- a/Engine/Plugins/Experimental/SceneState/Source/SceneState/Private/Tasks/SceneStateTaskInstance.cpp
+++ b/Engine/Plugins/Experimental/SceneState/Source/SceneState/Private/Tasks/SceneStateTaskInstance.cpp
@@ -2,18 +2,7 @@
 
 #include "Tasks/SceneStateTaskInstance.h"
 
-namespace UE::SceneState::Private
-{
-	
-uint16 GetNextInstanceId()
-{
-	static uint16 Id = 0;
-	return ++Id; // Starts with 1
-}
-
-} // UE::SceneState::Private
-
-FSceneStateTaskInstance::FSceneStateTaskInstance()
-	: InstanceId(UE::SceneState::Private::GetNextInstanceId())
+void FSceneStateTaskInstance::SetStatus(UE::SceneState::EExecutionStatus InStatus)
 {
+	Status = InStatus;
 }
diff --git a/Engine/Plugins/Experimental/SceneState/Source/SceneState/Public/SceneStateEnums.h b/Engine/Plugins/Experimental/SceneState/Source/SceneState/Public/SceneStateEnums.h
--- a/Engine/Plugins/Experimental/SceneState/Source/SceneState/Public/SceneStateEnums.h
+++ b/Engine/Plugins/Experimental/SceneState/Source/SceneState/Public/SceneStateEnums.h
@@ -39,15 +39,6 @@ enum class ESceneStateTaskFlags : uint8
 };
 ENUM_CLASS_FLAGS(ESceneStateTaskFlags);
 
-/** Defines the possible outcomes of a task */
-UENUM(BlueprintType)
-enum class ESceneStateTaskResult : uint8
-{
-	Undetermined,
-	Succeeded,
-	Failed,
-};
-
 /** Defines the possible ways a state machine will run */
 UENUM()
 enum class ESceneStateMachineRunMode : uint8
diff --git a/Engine/Plugins/Experimental/SceneState/Source/SceneState/Public/SceneStateInstance.h b/Engine/Plugins/Experimental/SceneState/Source/SceneState/Public/SceneStateInstance.h
--- a/Engine/Plugins/Experimental/SceneState/Source/SceneState/Public/SceneStateInstance.h
+++ b/Engine/Plugins/Experimental/SceneState/Source/SceneState/Public/SceneStateInstance.h
@@ -11,9 +11,17 @@ struct FSceneStateInstance
 {
 	GENERATED_BODY()
 
+	SCENESTATE_API FSceneStateInstance();
+
 	/** Time elapsed since the state entered */
 	float ElapsedTime = 0.f;
 
 	/** Current status of this state instance */
 	UE::SceneState::EExecutionStatus Status = UE::SceneState::EExecutionStatus::NotStarted;
+
+	uint16 GetInstanceId() const;
+
+private:
+	/** The id for this instance. This is used to differentiate state instances for a same state and consequently task instances for a same task*/
+	uint16 InstanceId;
 };
diff --git a/Engine/Plugins/Experimental/SceneState/Source/SceneState/Public/Tasks/SceneStateBlueprintableTask.h b/Engine/Plugins/Experimental/SceneState/Source/SceneState/Public/Tasks/SceneStateBlueprintableTask.h
--- a/Engine/Plugins/Experimental/SceneState/Source/SceneState/Public/Tasks/SceneStateBlueprintableTask.h
+++ b/Engine/Plugins/Experimental/SceneState/Source/SceneState/Public/Tasks/SceneStateBlueprintableTask.h
@@ -51,13 +51,7 @@ public:
 	USceneStateEventStream* GetEventStream() const;
 
 	UFUNCTION(BlueprintCallable, Category="Task")
-	void FinishTask(ESceneStateTaskResult InTaskResult = ESceneStateTaskResult::Succeeded);
-
-	UFUNCTION(BlueprintPure, Category="Task")
-	float GetTaskElapsedTime() const;
-
-	UFUNCTION(BlueprintPure, Category="Task")
-	ESceneStateTaskResult GetTaskResult() const;
+	void FinishTask();
 
 	//~ Begin ISceneStateEventHandlerProvider
 	virtual bool FindEventHandlerId(const FSceneStateEventSchemaHandle& InEventSchemaHandle, FGuid& OutHandlerId) const;
diff --git a/Engine/Plugins/Experimental/SceneState/Source/SceneState/Public/Tasks/SceneStateTaskExecutionContext.h b/Engine/Plugins/Experimental/SceneState/Source/SceneState/Public/Tasks/SceneStateTaskExecutionContext.h
--- a/Engine/Plugins/Experimental/SceneState/Source/SceneState/Public/Tasks/SceneStateTaskExecutionContext.h
+++ b/Engine/Plugins/Experimental/SceneState/Source/SceneState/Public/Tasks/SceneStateTaskExecutionContext.h
@@ -54,12 +54,15 @@ private:
 	/** Absolute index to the task */
 	uint16 TaskIndex;
 
+	/** Absolute index to the task's parent state */
+	uint16 StateIndex;
+
 	/**
-	 * Id assigned to the task instance to distinguish it from a different instance (within the same context) with the same task index.
+	 * Id assigned to the state instance to distinguish it from a different instance (within the same context) with the same state index.
 	 * This can happen if this context is held past its task instance (e.g. captured from a deferred lambda)
-	 * and there's an active new instance for the same task within the same context.
+	 * and there's an active new instance for the same state within the same context.
 	 */
-	uint32 TaskInstanceId;
+	uint16 StateInstanceId;
 };
 
 } // UE::SceneState
diff --git a/Engine/Plugins/Experimental/SceneState/Source/SceneState/Public/Tasks/SceneStateTaskInstance.h b/Engine/Plugins/Experimental/SceneState/Source/SceneState/Public/Tasks/SceneStateTaskInstance.h
--- a/Engine/Plugins/Experimental/SceneState/Source/SceneState/Public/Tasks/SceneStateTaskInstance.h
+++ b/Engine/Plugins/Experimental/SceneState/Source/SceneState/Public/Tasks/SceneStateTaskInstance.h
@@ -13,23 +13,16 @@ struct FSceneStateTaskInstance
 {
 	GENERATED_BODY()
 
-	SCENESTATE_API FSceneStateTaskInstance();
+	/** Sets the execution status of the Task */
+	void SetStatus(UE::SceneState::EExecutionStatus InStatus);
 
-	/** Elapsed Time */
-	float ElapsedTime = 0.f;
-
-	/** Current status of the Task */
-	UE::SceneState::EExecutionStatus Status = UE::SceneState::EExecutionStatus::NotStarted;
-
-	/** Result of the Task, when finished, otherwise undetermined */
-	ESceneStateTaskResult Result = ESceneStateTaskResult::Undetermined;
-
-	uint16 GetInstanceId() const
+	/** Gets the execution status of the task */
+	UE::SceneState::EExecutionStatus GetStatus() const
 	{
-		return InstanceId;
+		return Status;
 	}
 
 private:
-	/** The id for this instance. This is used to differentiate instances for the same task */
-	uint16 InstanceId;
+	/** Current status of the Task */
+	UE::SceneState::EExecutionStatus Status = UE::SceneState::EExecutionStatus::NotStarted;
 };
diff --git a/Engine/Plugins/Experimental/SceneState/Source/SceneStateBlueprintEditor/Private/DebugExecutors/SceneStateDebugTaskExecutor.cpp b/Engine/Plugins/Experimental/SceneState/Source/SceneStateBlueprintEditor/Private/DebugExecutors/SceneStateDebugTaskExecutor.cpp
--- a/Engine/Plugins/Experimental/SceneState/Source/SceneStateBlueprintEditor/Private/DebugExecutors/SceneStateDebugTaskExecutor.cpp
+++ b/Engine/Plugins/Experimental/SceneState/Source/SceneStateBlueprintEditor/Private/DebugExecutors/SceneStateDebugTaskExecutor.cpp
@@ -96,7 +96,7 @@ void FDebugTaskExecutor::OnExit(const FSceneStateExecutionContext& InExecutionCo
 void FDebugTaskExecutor::ConditionallyExit(FStructView InTaskInstance)
 {
 	const FSceneStateTaskInstance* TaskInstance = InTaskInstance.GetPtr<FSceneStateTaskInstance>();
-	if (TaskInstance && TaskInstance->Status == EExecutionStatus::Finished)
+	if (TaskInstance && TaskInstance->GetStatus() == EExecutionStatus::Finished)
 	{
 		Exit();
 	}
diff --git a/Engine/Plugins/Experimental/SceneState/Source/SceneStateMachineEditor/Private/Widgets/SSceneStateMachineTaskNode.cpp b/Engine/Plugins/Experimental/SceneState/Source/SceneStateMachineEditor/Private/Widgets/SSceneStateMachineTaskNode.cpp
--- a/Engine/Plugins/Experimental/SceneState/Source/SceneStateMachineEditor/Private/Widgets/SSceneStateMachineTaskNode.cpp
+++ b/Engine/Plugins/Experimental/SceneState/Source/SceneStateMachineEditor/Private/Widgets/SSceneStateMachineTaskNode.cpp
@@ -201,17 +201,16 @@ void SStateMachineTaskNode::GetNodeInfoPopups(FNodeInfoContext* InContext, TArra
 		[&OutPopups, &ResultStatus, &StatusColors = StatusColors](const FSceneStateTaskInstance& InTaskInstance)
 		{
 			// Get the lowest status to show
-			if (InTaskInstance.Status <= ResultStatus.Get(EExecutionStatus::Finished))
+			if (InTaskInstance.GetStatus() <= ResultStatus.Get(EExecutionStatus::Finished))
 			{
-				ResultStatus = InTaskInstance.Status;
+				ResultStatus = InTaskInstance.GetStatus();
 			}
 
 			FText Message;
-			switch (InTaskInstance.Status)
+			switch (InTaskInstance.GetStatus())
 			{
 			case EExecutionStatus::Running:
-				Message = FText::Format(LOCTEXT("RunningStatusFormat", "Active for {0} s")
-					, FText::AsNumber(InTaskInstance.ElapsedTime, &FStateMachineEditorStyle::Get().GetDefaultNumberFormat()));
+				Message = LOCTEXT("RunningStatusFormat", "Running");
 				break;
 
 			case EExecutionStatus::Finished:
@@ -222,7 +221,7 @@ void SStateMachineTaskNode::GetNodeInfoPopups(FNodeInfoContext* InContext, TArra
 				return;
 			}
 
-			OutPopups.Emplace(nullptr, StatusColors[InTaskInstance.Status], Message.ToString());
+			OutPopups.Emplace(nullptr, StatusColors[InTaskInstance.GetStatus()], Message.ToString());
 		});
 
 	// If the result status was never gotten (i.e. no Task Instances), set to Not-Started
diff --git a/Engine/Plugins/Experimental/SceneState/Source/SceneStateTasks/Private/SceneStateDelayTask.cpp b/Engine/Plugins/Experimental/SceneState/Source/SceneStateTasks/Private/SceneStateDelayTask.cpp
--- a/Engine/Plugins/Experimental/SceneState/Source/SceneStateTasks/Private/SceneStateDelayTask.cpp
+++ b/Engine/Plugins/Experimental/SceneState/Source/SceneStateTasks/Private/SceneStateDelayTask.cpp
@@ -17,7 +17,9 @@ const UScriptStruct* FSceneStateDelayTask::OnGetTaskInstanceType() const
 
 void FSceneStateDelayTask::OnStart(const FSceneStateExecutionContext& InContext, FStructView InTaskInstance) const
 {
-	const FInstanceDataType& Instance = InTaskInstance.Get<FInstanceDataType>();
+	FInstanceDataType& Instance = InTaskInstance.Get<FInstanceDataType>();
+	Instance.ElapsedTime = 0.f;
+
 	if (Instance.Delay < UE_SMALL_NUMBER)
 	{
 		Finish(InContext, InTaskInstance);
@@ -26,7 +28,9 @@ void FSceneStateDelayTask::OnStart(const FSceneStateExecutionContext& InContext,
 
 void FSceneStateDelayTask::OnTick(const FSceneStateExecutionContext& InContext, FStructView InTaskInstance, float InDeltaSeconds) const
 {
-	const FInstanceDataType& Instance = InTaskInstance.Get<FInstanceDataType>();
+	FInstanceDataType& Instance = InTaskInstance.Get<FInstanceDataType>();
+	Instance.ElapsedTime += InDeltaSeconds;
+
 	if (Instance.ElapsedTime >= Instance.Delay)
 	{
 		Finish(InContext, InTaskInstance);
diff --git a/Engine/Plugins/Experimental/SceneState/Source/SceneStateTasks/Public/SceneStateDelayTask.h b/Engine/Plugins/Experimental/SceneState/Source/SceneStateTasks/Public/SceneStateDelayTask.h
--- a/Engine/Plugins/Experimental/SceneState/Source/SceneStateTasks/Public/SceneStateDelayTask.h
+++ b/Engine/Plugins/Experimental/SceneState/Source/SceneStateTasks/Public/SceneStateDelayTask.h
@@ -11,6 +11,9 @@ struct FSceneStateDelayTaskInstance : public FSceneStateTaskInstance
 {
 	GENERATED_BODY()
 
+	/** The amount of time elapsed */
+	float ElapsedTime = 0.f;
+
 	/** The amount to wait for, in seconds */
 	UPROPERTY(EditAnywhere, Category="Delay", meta=(ClampMin="0"))
 	float Delay = 0.5f;
