Motion Design: added property ref comments, and cleaned up the property ref functions and added recursion support

diff --git a/Engine/Plugins/Experimental/SceneState/Source/SceneState/Private/SceneStatePropertyUtils.cpp b/Engine/Plugins/Experimental/SceneState/Source/SceneState/Private/SceneStatePropertyUtils.cpp
--- a/Engine/Plugins/Experimental/SceneState/Source/SceneState/Private/SceneStatePropertyUtils.cpp
+++ b/Engine/Plugins/Experimental/SceneState/Source/SceneState/Private/SceneStatePropertyUtils.cpp
@@ -8,30 +8,42 @@
 namespace UE::SceneState
 {
 
-const FSceneStateBindingResolvedReference* FindResolvedReference(const FSceneStateExecutionContext& InContext, const FSceneStatePropertyReference& InPropertyReference)
+bool ResolveProperty(const FSceneStateExecutionContext& InContext, const FSceneStatePropertyReference& InPropertyReference, FResolvePropertyResult& OutResult)
 {
-	if (const FSceneStateBindingCollection* BindingCollection = InContext.GetBindingCollection())
+	const FSceneStateBindingCollection* BindingCollection = InContext.GetBindingCollection();
+	if (!BindingCollection)
 	{
-		return BindingCollection->FindResolvedReference(InPropertyReference);
+		return false;
 	}
-	return nullptr;
-}
 
-uint8* ResolveProperty(const FSceneStateExecutionContext& InContext, const FSceneStateBindingResolvedReference& InResolvedReference)
-{
-	const FSceneStateBindingCollection* BindingCollection = InContext.GetBindingCollection();
-	if (!BindingCollection)
+	const FSceneStateBindingResolvedReference* ResolvedReference = BindingCollection->FindResolvedReference(InPropertyReference);
+	if (!ResolvedReference || !ResolvedReference->SourceLeafProperty)
 	{
-		return nullptr;
+		return false;
 	}
 
-	const FPropertyBindingDataView SourceDataView = InContext.FindDataView(InResolvedReference.SourceDataHandle);
+	const FPropertyBindingDataView SourceDataView = InContext.FindDataView(ResolvedReference->SourceDataHandle);
 	if (!SourceDataView.IsValid())
 	{
-		return nullptr;
+		return false;
+	}
+
+	uint8* const ValuePtr = BindingCollection->ResolveProperty(*ResolvedReference, SourceDataView);
+	if (!ValuePtr)
+	{
+		return false;
+	}
+
+	// If the Source is a Property Reference, recursively resolve property
+	if (IsPropertyReference(ResolvedReference->SourceLeafProperty))
+	{
+		const FSceneStatePropertyReference& SourcePropertyReference = *reinterpret_cast<const FSceneStatePropertyReference*>(ValuePtr);
+		return ResolveProperty(InContext, SourcePropertyReference, OutResult);
 	}
 
-	return BindingCollection->ResolveProperty(InResolvedReference, SourceDataView);
+	OutResult.ValuePtr = ValuePtr;
+	OutResult.ResolvedReference = ResolvedReference;
+	return true;
 }
 
 } // UE::SceneState
diff --git a/Engine/Plugins/Experimental/SceneState/Source/SceneState/Public/SceneStatePropertyUtils.h b/Engine/Plugins/Experimental/SceneState/Source/SceneState/Public/SceneStatePropertyUtils.h
--- a/Engine/Plugins/Experimental/SceneState/Source/SceneState/Public/SceneStatePropertyUtils.h
+++ b/Engine/Plugins/Experimental/SceneState/Source/SceneState/Public/SceneStatePropertyUtils.h
@@ -10,31 +10,31 @@ struct FSceneStatePropertyReference;
 
 namespace UE::SceneState
 {
+	struct FResolvePropertyResult
+	{
+		/** the resolved value address */
+		uint8* ValuePtr = nullptr;
+		/** the resolved property reference mapped to the value address (can differ from the resolved reference mapped to InPropertyReference in recursive scenarios) */
+		const FSceneStateBindingResolvedReference* ResolvedReference = nullptr;
+	};
 	/**
-	 * Finds the resolved reference for the given property reference
-	 * @param InContext the execution context containing the binding collection
-	 * @param InPropertyReference the property reference to resolve
-	 * @return the resolved reference if found, null otherwise
-	 */
-	SCENESTATE_API const FSceneStateBindingResolvedReference* FindResolvedReference(const FSceneStateExecutionContext& InContext, const FSceneStatePropertyReference& InPropertyReference);
-
-	/**
-	 * Resolves the given reference and execution context to the value address of the property
+	 * Resolves the given reference and execution context to the value address of the property, with recursion if the source is a property reference too.
 	 * @param InContext the execution context containing the binding collection and way to retrieve the data views
-	 * @param InResolvedReference the property reference to resolve access for
-	 * @return value address of the property if successful, null otherwise
+	 * @param InPropertyReference the property reference to resolve access for
+	 * @param OutResult if successful, a result containing a valid resolved reference and a valid value address.
+	 * @return whether the operation succeeded
 	 */
-	SCENESTATE_API uint8* ResolveProperty(const FSceneStateExecutionContext& InContext, const FSceneStateBindingResolvedReference& InResolvedReference);
+	SCENESTATE_API bool ResolveProperty(const FSceneStateExecutionContext& InContext, const FSceneStatePropertyReference& InPropertyReference, FResolvePropertyResult& OutResult);
 
 	template<class T>
 	T* ResolveProperty(const FSceneStateExecutionContext& InContext, const FSceneStatePropertyReference& InPropertyReference)
 	{
-		const FSceneStateBindingResolvedReference* ResolvedReference = FindResolvedReference(InContext, InPropertyReference);
-		if (ResolvedReference && ResolvedReference->SourceLeafProperty && TDataTypeHelper<T>::IsValid(ResolvedReference->SourceLeafProperty))
+		FResolvePropertyResult Result;
+		if (ResolveProperty(InContext, InPropertyReference, Result) && TDataTypeHelper<T>::IsValid(Result.ResolvedReference->SourceLeafProperty))
 		{
-			return reinterpret_cast<T*>(ResolveProperty(InContext, *ResolvedReference));
+			return reinterpret_cast<T*>(Result.ValuePtr);
 		}
-		return nullptr;
+		return false;
 	}
 
 } // UE::SceneState
diff --git a/Engine/Plugins/Experimental/SceneState/Source/SceneStateBinding/Public/SceneStateBindingCollection.h b/Engine/Plugins/Experimental/SceneState/Source/SceneStateBinding/Public/SceneStateBindingCollection.h
--- a/Engine/Plugins/Experimental/SceneState/Source/SceneStateBinding/Public/SceneStateBindingCollection.h
+++ b/Engine/Plugins/Experimental/SceneState/Source/SceneStateBinding/Public/SceneStateBindingCollection.h
@@ -36,7 +36,11 @@ struct FSceneStateBindingCollection : public FPropertyBindingBindingCollection
 	/** Retrieves the resolved reference for the given property reference */
 	SCENESTATEBINDING_API const FSceneStateBindingResolvedReference* FindResolvedReference(const FSceneStatePropertyReference& InPropertyReference) const;
 
-	/** Resolves the property access for the given reference and data view */
+	/**
+	 * Resolves the property access for the given reference and data view
+	 * @note This does not deal with recursive resolution (for cases where the reference is to another reference).
+	 * Prefer using UE::SceneState::ResolveProperty instead.
+	 */
 	SCENESTATEBINDING_API uint8* ResolveProperty(const FSceneStateBindingResolvedReference& InResolvedReference, FPropertyBindingDataView InDataView) const;
 
 	//~ Begin FPropertyBindingBindingCollection
diff --git a/Engine/Plugins/Experimental/SceneState/Source/SceneStateBinding/Public/SceneStatePropertyReference.h b/Engine/Plugins/Experimental/SceneState/Source/SceneStateBinding/Public/SceneStatePropertyReference.h
--- a/Engine/Plugins/Experimental/SceneState/Source/SceneStateBinding/Public/SceneStatePropertyReference.h
+++ b/Engine/Plugins/Experimental/SceneState/Source/SceneStateBinding/Public/SceneStatePropertyReference.h
@@ -7,6 +7,10 @@
 
 /**
  * Property References allows the binding system to handle a property by reference rather than by copy.
+ * A reference holds a single 'ReferenceIndex' that points to the index of both a 'Reference' and its 'Resolved Reference'.
+ * To get the property by reference, the binding collection (to get the data handle) and the execution context (to get the data view) are required.
+ * @see FSceneStateBindingCollection, FSceneStateBindingReference, FSceneStateBindingResolvedReference.
+ *
  * This is useful for tasks that need to write to external properties (e.g. as a form of output).
  * The expected type of the reference should be set in "RefType" meta specifier.
  * Meta specifiers for the type:
diff --git a/Engine/Plugins/Experimental/SceneState/Source/SceneStateTasks/Public/Setters/SceneStateSetterUtils.h b/Engine/Plugins/Experimental/SceneState/Source/SceneStateTasks/Public/Setters/SceneStateSetterUtils.h
--- a/Engine/Plugins/Experimental/SceneState/Source/SceneStateTasks/Public/Setters/SceneStateSetterUtils.h
+++ b/Engine/Plugins/Experimental/SceneState/Source/SceneStateTasks/Public/Setters/SceneStateSetterUtils.h
@@ -66,19 +66,12 @@ namespace UE::SceneState
 	{
 		typename InTaskType::FInstanceDataType& Instance = InTaskInstance.Get<typename InTaskType::FInstanceDataType>();
 
-		const FSceneStateBindingResolvedReference* ResolvedReference = FindResolvedReference(InContext, Instance.Target);
-		if (!ResolvedReference || !ResolvedReference->SourceLeafProperty)
+		FResolvePropertyResult Result;
+		if (ResolveProperty(InContext, Instance.Target, Result))
 		{
-			return false;
+			return (SetValue<InTargetTypes>(Result.ValuePtr, *Result.ResolvedReference, Instance.Value) || ...);
 		}
-
-		uint8* const ValuePtr = ResolveProperty(InContext, *ResolvedReference);
-		if (!ValuePtr)
-		{
-			return false;
-		}
-
-		return (SetValue<InTargetTypes>(ValuePtr, *ResolvedReference, Instance.Value) || ...);
+		return false;
 	}
 
 } // UE::SceneState
