Motion Design: fix scene state possible recursion + function wrongly returning false (instead of nullptr)

diff --git a/Engine/Plugins/Experimental/SceneState/Source/SceneState/Private/SceneStatePropertyUtils.cpp b/Engine/Plugins/Experimental/SceneState/Source/SceneState/Private/SceneStatePropertyUtils.cpp
--- a/Engine/Plugins/Experimental/SceneState/Source/SceneState/Private/SceneStatePropertyUtils.cpp
+++ b/Engine/Plugins/Experimental/SceneState/Source/SceneState/Private/SceneStatePropertyUtils.cpp
@@ -8,7 +8,11 @@
 namespace UE::SceneState
 {
 
-bool ResolveProperty(const FSceneStateExecutionContext& InContext, const FSceneStatePropertyReference& InPropertyReference, FResolvePropertyResult& OutResult)
+namespace Private
+{
+
+/** Resolves a property reference without recursing/iterating further if the source of that reference is another property reference */
+bool ResolvePropertyInternal(const FSceneStateExecutionContext& InContext, const FSceneStatePropertyReference& InPropertyReference, FResolvePropertyResult& OutResult)
 {
 	const FSceneStateBindingCollection* BindingCollection = InContext.GetBindingCollection();
 	if (!BindingCollection)
@@ -34,15 +38,34 @@ bool ResolveProperty(const FSceneStateExecutionContext& InContext, const FSceneS
 		return false;
 	}
 
-	// If the Source is a Property Reference, recursively resolve property
-	if (IsPropertyReference(ResolvedReference->SourceLeafProperty))
+	OutResult.ValuePtr = ValuePtr;
+	OutResult.ResolvedReference = ResolvedReference;
+	return true;
+}
+
+} // UE::SceneState::Private
+	
+bool ResolveProperty(const FSceneStateExecutionContext& InContext, const FSceneStatePropertyReference& InPropertyReference, FResolvePropertyResult& OutResult)
+{
+	FResolvePropertyResult Result;
+	if (!Private::ResolvePropertyInternal(InContext, InPropertyReference, Result))
 	{
-		const FSceneStatePropertyReference& SourcePropertyReference = *reinterpret_cast<const FSceneStatePropertyReference*>(ValuePtr);
-		return ResolveProperty(InContext, SourcePropertyReference, OutResult);
+		return false;
 	}
 
-	OutResult.ValuePtr = ValuePtr;
-	OutResult.ResolvedReference = ResolvedReference;
+	constexpr int32 MaxIterations = 100;
+	int32 IterationCount = 0;
+
+	// If the Source is a Property Reference, follow the reference chain until a non-property reference is found (or a safe max iteration limit is hit)
+	while (ensure(IterationCount++ < MaxIterations) && IsPropertyReference(Result.ResolvedReference->SourceLeafProperty))
+	{
+		if (!Private::ResolvePropertyInternal(InContext, InPropertyReference, Result))
+		{
+			return false;
+		}
+	}
+
+	OutResult = MoveTemp(Result);
 	return true;
 }
 
diff --git a/Engine/Plugins/Experimental/SceneState/Source/SceneState/Public/SceneStatePropertyUtils.h b/Engine/Plugins/Experimental/SceneState/Source/SceneState/Public/SceneStatePropertyUtils.h
--- a/Engine/Plugins/Experimental/SceneState/Source/SceneState/Public/SceneStatePropertyUtils.h
+++ b/Engine/Plugins/Experimental/SceneState/Source/SceneState/Public/SceneStatePropertyUtils.h
@@ -18,7 +18,7 @@ namespace UE::SceneState
 		const FSceneStateBindingResolvedReference* ResolvedReference = nullptr;
 	};
 	/**
-	 * Resolves the given reference and execution context to the value address of the property, with recursion if the source is a property reference too.
+	 * Resolves the given reference and execution context to the value address of the property, and keeps iterating if the source is a property reference too.
 	 * @param InContext the execution context containing the binding collection and way to retrieve the data views
 	 * @param InPropertyReference the property reference to resolve access for
 	 * @param OutResult if successful, a result containing a valid resolved reference and a valid value address.
@@ -34,7 +34,7 @@ namespace UE::SceneState
 		{
 			return reinterpret_cast<T*>(Result.ValuePtr);
 		}
-		return false;
+		return nullptr;
 	}
 
 } // UE::SceneState
