Motion Design: fix resolve property calling recursively with same parameters

diff --git a/Engine/Plugins/Experimental/SceneState/Source/SceneState/Private/SceneStatePropertyUtils.cpp b/Engine/Plugins/Experimental/SceneState/Source/SceneState/Private/SceneStatePropertyUtils.cpp
--- a/Engine/Plugins/Experimental/SceneState/Source/SceneState/Private/SceneStatePropertyUtils.cpp
+++ b/Engine/Plugins/Experimental/SceneState/Source/SceneState/Private/SceneStatePropertyUtils.cpp
@@ -59,7 +59,8 @@ bool ResolveProperty(const FSceneStateExecutionContext& InContext, const FSceneS
 	// If the Source is a Property Reference, follow the reference chain until a non-property reference is found (or a safe max iteration limit is hit)
 	while (ensure(IterationCount++ < MaxIterations) && IsPropertyReference(Result.ResolvedReference->SourceLeafProperty))
 	{
-		if (!Private::ResolvePropertyInternal(InContext, InPropertyReference, Result))
+		const FSceneStatePropertyReference& SourcePropertyReference = *reinterpret_cast<FSceneStatePropertyReference*>(Result.ValuePtr);
+		if (!Private::ResolvePropertyInternal(InContext, SourcePropertyReference, Result))
 		{
 			return false;
 		}
