Motion Design: added blueprint property reference support, so that Parameters can make use of them.
Even though BP Tasks also have this type available, it is not yet fully usable. Further work is needed to get this.

diff --git a/Engine/Plugins/Experimental/AvalancheSceneState/Source/AvalancheSceneState/Private/RemoteControl/AvaSceneStateRCTaskBinding.cpp b/Engine/Plugins/Experimental/AvalancheSceneState/Source/AvalancheSceneState/Private/RemoteControl/AvaSceneStateRCTaskBinding.cpp
--- a/Engine/Plugins/Experimental/AvalancheSceneState/Source/AvalancheSceneState/Private/RemoteControl/AvaSceneStateRCTaskBinding.cpp
+++ b/Engine/Plugins/Experimental/AvalancheSceneState/Source/AvalancheSceneState/Private/RemoteControl/AvaSceneStateRCTaskBinding.cpp
@@ -4,9 +4,9 @@
 #include "RemoteControl/AvaSceneStateRCTask.h"
 
 #if WITH_EDITOR
-void FAvaSceneStateRCTaskBinding::VisitBindingDescs(FConstStructView InTaskInstance, TFunctionRef<void(const UE::SceneState::FTaskBindingDesc&)> InFunctor) const
+void FAvaSceneStateRCTaskBinding::VisitBindingDescs(FStructView InTaskInstance, TFunctionRef<void(const UE::SceneState::FTaskBindingDesc&)> InFunctor) const
 {
-	const FInstanceDataType& Instance = InTaskInstance.Get<const FInstanceDataType>();
+	FInstanceDataType& Instance = InTaskInstance.Get<FInstanceDataType>();
 
 	if (Instance.ControllerValues.IsValid())
 	{
@@ -14,7 +14,7 @@ void FAvaSceneStateRCTaskBinding::VisitBindingDescs(FConstStructView InTaskInsta
 
 		BindingDesc.Id = Instance.ControllerValuesId;
 		BindingDesc.Name = TEXT("ControllerValues");
-		BindingDesc.Struct = Instance.ControllerValues.GetPropertyBagStruct();
+		BindingDesc.DataView = Instance.ControllerValues.GetMutableValue();
 		BindingDesc.DataIndex = CONTROLLER_VALUES_DATA_INDEX;
 
 		InFunctor(BindingDesc);
diff --git a/Engine/Plugins/Experimental/AvalancheSceneState/Source/AvalancheSceneState/Public/RemoteControl/AvaSceneStateRCTaskBinding.h b/Engine/Plugins/Experimental/AvalancheSceneState/Source/AvalancheSceneState/Public/RemoteControl/AvaSceneStateRCTaskBinding.h
--- a/Engine/Plugins/Experimental/AvalancheSceneState/Source/AvalancheSceneState/Public/RemoteControl/AvaSceneStateRCTaskBinding.h
+++ b/Engine/Plugins/Experimental/AvalancheSceneState/Source/AvalancheSceneState/Public/RemoteControl/AvaSceneStateRCTaskBinding.h
@@ -20,7 +20,7 @@ struct FAvaSceneStateRCTaskBinding : public FSceneStateTaskBindingExtension
 protected:
 	//~ Begin FSceneStateTaskBindingExtension
 #if WITH_EDITOR
-	virtual void VisitBindingDescs(FConstStructView InTaskInstance, TFunctionRef<void(const UE::SceneState::FTaskBindingDesc&)> InFunctor) const override;
+	virtual void VisitBindingDescs(FStructView InTaskInstance, TFunctionRef<void(const UE::SceneState::FTaskBindingDesc&)> InFunctor) const override;
 	virtual void SetBindingBatch(uint16 InDataIndex, uint16 InBatchIndex) override;
 	virtual bool FindDataById(FStructView InTaskInstance, const FGuid& InStructId, FStructView& OutDataView, uint16& OutDataIndex) const override;
 #endif
diff --git a/Engine/Plugins/Experimental/SceneState/Source/SceneState/Private/SceneStateExecutionContext.cpp b/Engine/Plugins/Experimental/SceneState/Source/SceneState/Private/SceneStateExecutionContext.cpp
--- a/Engine/Plugins/Experimental/SceneState/Source/SceneState/Private/SceneStateExecutionContext.cpp
+++ b/Engine/Plugins/Experimental/SceneState/Source/SceneState/Private/SceneStateExecutionContext.cpp
@@ -61,6 +61,15 @@ USceneStateEventStream* FSceneStateExecutionContext::GetEventStream() const
 	return nullptr;
 }
 
+void FSceneStateExecutionContext::Setup(const FSceneStateExecutionContext& InParentContext)
+{
+	Setup(InParentContext.GetRootState());
+	check(RootState);
+
+	const TSharedRef<UE::SceneState::FExecutionContextRegistry> ContextRegistry = RootState->GetContextRegistry();
+	ParentHandle = ContextRegistry->FindHandle(InParentContext);
+}
+
 void FSceneStateExecutionContext::Setup(USceneStateObject* InRootState)
 {
 	Reset();
@@ -154,6 +163,12 @@ FPropertyBindingDataView FSceneStateExecutionContext::FindDataView(const FSceneS
 		break;
 	}
 
+	// The data handle could be exclusively on the parent context, so call FindDataView on the parent too
+	if (const FSceneStateExecutionContext* ParentContext = FindParentContext())
+	{
+		return ParentContext->FindDataView(InDataHandle);
+	}
+
 	return FPropertyBindingDataView();
 }
 
@@ -645,6 +660,15 @@ void FSceneStateExecutionContext::RemoveTransitionInstance(const FSceneStateTran
 	}
 }
 
+const FSceneStateExecutionContext* FSceneStateExecutionContext::FindParentContext() const
+{
+	if (TSharedPtr<UE::SceneState::FExecutionContextRegistry> ContextRegistry = ContextRegistryWeak.Pin())
+	{
+		return ContextRegistry->FindContext(ParentHandle);
+	}
+	return nullptr;
+}
+
 bool FSceneStateExecutionContext::GetStateIndex(const FSceneState& InState, uint16& OutIndex) const
 {
 	if (GeneratedClass)
diff --git a/Engine/Plugins/Experimental/SceneState/Source/SceneState/Private/Tasks/SceneStateMachineTask.cpp b/Engine/Plugins/Experimental/SceneState/Source/SceneState/Private/Tasks/SceneStateMachineTask.cpp
--- a/Engine/Plugins/Experimental/SceneState/Source/SceneState/Private/Tasks/SceneStateMachineTask.cpp
+++ b/Engine/Plugins/Experimental/SceneState/Source/SceneState/Private/Tasks/SceneStateMachineTask.cpp
@@ -42,7 +42,7 @@ void FSceneStateMachineTask::OnStart(const FSceneStateExecutionContext& InContex
 		return;
 	}
 
-	Instance.ExecutionContext.Setup(InContext.GetRootState());
+	Instance.ExecutionContext.Setup(InContext);
 
 	StateMachine->Setup(Instance.ExecutionContext);
 
diff --git a/Engine/Plugins/Experimental/SceneState/Source/SceneState/Private/Tasks/SceneStateMachineTaskBinding.cpp b/Engine/Plugins/Experimental/SceneState/Source/SceneState/Private/Tasks/SceneStateMachineTaskBinding.cpp
--- a/Engine/Plugins/Experimental/SceneState/Source/SceneState/Private/Tasks/SceneStateMachineTaskBinding.cpp
+++ b/Engine/Plugins/Experimental/SceneState/Source/SceneState/Private/Tasks/SceneStateMachineTaskBinding.cpp
@@ -4,9 +4,9 @@
 #include "Tasks/SceneStateMachineTask.h"
 
 #if WITH_EDITOR
-void FSceneStateMachineTaskBinding::VisitBindingDescs(FConstStructView InTaskInstance, TFunctionRef<void(const UE::SceneState::FTaskBindingDesc&)> InFunctor) const
+void FSceneStateMachineTaskBinding::VisitBindingDescs(FStructView InTaskInstance, TFunctionRef<void(const UE::SceneState::FTaskBindingDesc&)> InFunctor) const
 {
-	const FInstanceDataType& Instance = InTaskInstance.Get<const FInstanceDataType>();
+	FInstanceDataType& Instance = InTaskInstance.Get<FInstanceDataType>();
 
 	if (Instance.Parameters.IsValid())
 	{
@@ -14,7 +14,7 @@ void FSceneStateMachineTaskBinding::VisitBindingDescs(FConstStructView InTaskIns
 
 		BindingDesc.Id = Instance.ParametersId;
 		BindingDesc.Name = TEXT("State Machine Task Parameters");
-		BindingDesc.Struct = Instance.Parameters.GetPropertyBagStruct();
+		BindingDesc.DataView = Instance.Parameters.GetMutableValue();
 		BindingDesc.DataIndex = PARAMETERS_DATA_INDEX;
 
 		InFunctor(BindingDesc);
diff --git a/Engine/Plugins/Experimental/SceneState/Source/SceneState/Public/SceneStateExecutionContext.h b/Engine/Plugins/Experimental/SceneState/Source/SceneState/Public/SceneStateExecutionContext.h
--- a/Engine/Plugins/Experimental/SceneState/Source/SceneState/Public/SceneStateExecutionContext.h
+++ b/Engine/Plugins/Experimental/SceneState/Source/SceneState/Public/SceneStateExecutionContext.h
@@ -3,6 +3,7 @@
 #pragma once
 
 #include "Containers/Array.h"
+#include "SceneStateExecutionContextHandle.h"
 #include "SceneStateInstance.h"
 #include "SceneStateMachineInstance.h"
 #include "SceneStatePropertyUtils.h"
@@ -79,6 +80,12 @@ struct FSceneStateExecutionContext
 	/** Gets the Event Stream from the Root State if available */
 	SCENESTATE_API USceneStateEventStream* GetEventStream() const;
 
+	/**
+	 * Initialization for a context with a parent context.
+	 * @param InParentContext the parent of this execution context
+	 */
+	SCENESTATE_API void Setup(const FSceneStateExecutionContext& InParentContext);
+
 	/**
 	 * Initialization for the Context.
 	 * Through the given root state it pre-allocates the instance data that will be used
@@ -201,6 +208,9 @@ struct FSceneStateExecutionContext
 	void RemoveTransitionInstance(const FSceneStateTransition& InTransition) const;
 
 private:
+	/** Retrieves the parent execution context if existing. Can be null */
+	const FSceneStateExecutionContext* FindParentContext() const;
+
 	/** Retrieves the Absolute Index of a State in the State array */
 	bool GetStateIndex(const FSceneState& InState, uint16& OutIndex) const;
 
@@ -246,6 +256,9 @@ private:
 	UPROPERTY()
 	mutable TMap<uint16, FSceneStateTransitionInstance> TransitionInstances;
 
+	/** Context handle to the parent execution context. Can be invalid. */
+	UE::SceneState::FExecutionContextHandle ParentHandle;
+
 	/** Weak reference to the registry containing this context */
 	TWeakPtr<UE::SceneState::FExecutionContextRegistry> ContextRegistryWeak;
 };
diff --git a/Engine/Plugins/Experimental/SceneState/Source/SceneState/Public/SceneStateExecutionContextHandle.h b/Engine/Plugins/Experimental/SceneState/Source/SceneState/Public/SceneStateExecutionContextHandle.h
--- a/Engine/Plugins/Experimental/SceneState/Source/SceneState/Public/SceneStateExecutionContextHandle.h
+++ b/Engine/Plugins/Experimental/SceneState/Source/SceneState/Public/SceneStateExecutionContextHandle.h
@@ -3,6 +3,7 @@
 #pragma once
 
 #include "HAL/Platform.h"
+#include "Templates/TypeHash.h"
 
 namespace UE::SceneState
 {
@@ -23,11 +24,11 @@ public:
 
 	friend uint32 GetTypeHash(FExecutionContextHandle InHandle)
 	{
-		return InHandle.Id;
+		return ::GetTypeHash(InHandle.Id);
 	}
 
 private:
-	uint32 Id = 0;
+	uint64 Id = 0;
 
 	friend class FExecutionContextRegistry;
 };
diff --git a/Engine/Plugins/Experimental/SceneState/Source/SceneState/Public/Tasks/SceneStateMachineTaskBinding.h b/Engine/Plugins/Experimental/SceneState/Source/SceneState/Public/Tasks/SceneStateMachineTaskBinding.h
--- a/Engine/Plugins/Experimental/SceneState/Source/SceneState/Public/Tasks/SceneStateMachineTaskBinding.h
+++ b/Engine/Plugins/Experimental/SceneState/Source/SceneState/Public/Tasks/SceneStateMachineTaskBinding.h
@@ -20,7 +20,7 @@ struct FSceneStateMachineTaskBinding : public FSceneStateTaskBindingExtension
 protected:
 	//~ Begin FSceneStateTaskBindingExtension
 #if WITH_EDITOR
-	virtual void VisitBindingDescs(FConstStructView InTaskInstance, TFunctionRef<void(const UE::SceneState::FTaskBindingDesc&)> InFunctor) const override;
+	virtual void VisitBindingDescs(FStructView InTaskInstance, TFunctionRef<void(const UE::SceneState::FTaskBindingDesc&)> InFunctor) const override;
 	virtual void SetBindingBatch(uint16 InDataIndex, uint16 InBatchIndex) override;
 	virtual bool FindDataById(FStructView InTaskInstance, const FGuid& InStructId, FStructView& OutDataView, uint16& OutDataIndex) const override;
 #endif
diff --git a/Engine/Plugins/Experimental/SceneState/Source/SceneState/Public/Tasks/SceneStateTaskBindingExtension.h b/Engine/Plugins/Experimental/SceneState/Source/SceneState/Public/Tasks/SceneStateTaskBindingExtension.h
--- a/Engine/Plugins/Experimental/SceneState/Source/SceneState/Public/Tasks/SceneStateTaskBindingExtension.h
+++ b/Engine/Plugins/Experimental/SceneState/Source/SceneState/Public/Tasks/SceneStateTaskBindingExtension.h
@@ -23,8 +23,8 @@ struct FTaskBindingDesc
 	/** Name of the struct (used for debugging, logging, cosmetic). */
 	FName Name;
 
-	/** The type of the struct or class. */
-	TObjectPtr<const UStruct> Struct = nullptr;
+	/** Template data view to the struct or class data  */
+	FPropertyBindingDataView DataView;
 
 	/** Unique Index for this Custom Desc to identify at Runtime */
 	uint16 DataIndex = FSceneStateRange::InvalidIndex;
@@ -45,7 +45,7 @@ struct FSceneStateTaskBindingExtension
 	 * @param InTaskInstance the task instance with binding structs
 	 * @param InFunctor the functor to execute for every available binding desc
 	 */
-	virtual void VisitBindingDescs(FConstStructView InTaskInstance, TFunctionRef<void(const UE::SceneState::FTaskBindingDesc&)> InFunctor) const
+	virtual void VisitBindingDescs(FStructView InTaskInstance, TFunctionRef<void(const UE::SceneState::FTaskBindingDesc&)> InFunctor) const
 	{
 	}
 
diff --git a/Engine/Plugins/Experimental/SceneState/Source/SceneStateBinding/Private/SceneStateBlueprintPropertyReference.cpp b/Engine/Plugins/Experimental/SceneState/Source/SceneStateBinding/Private/SceneStateBlueprintPropertyReference.cpp
new file mode 100644
--- /dev/null
+++ b/Engine/Plugins/Experimental/SceneState/Source/SceneStateBinding/Private/SceneStateBlueprintPropertyReference.cpp
@@ -0,0 +1,129 @@
+// Copyright Epic Games, Inc. All Rights Reserved.
+
+#include "SceneStateBlueprintPropertyReference.h"
+
+#if WITH_EDITOR
+#include "EdGraph/EdGraphPin.h"
+#include "EdGraphSchema_K2.h"
+#endif
+
+#if WITH_EDITOR
+FSceneStateBlueprintPropertyReference FSceneStateBlueprintPropertyReference::BuildFromPinType(const FEdGraphPinType& InPinType)
+{
+	FSceneStateBlueprintPropertyReference PropertyReference;
+
+	switch (InPinType.ContainerType)
+	{
+	case EPinContainerType::Array:
+		PropertyReference.bIsReferenceToArray = true;
+		break;
+	case EPinContainerType::Set:
+		ensureMsgf(false, TEXT("Unsupported container type [Set]"));
+		break;
+	case EPinContainerType::Map:
+		ensureMsgf(false, TEXT("Unsupported container type [Map]"));
+		break;
+	default:
+		break;
+	}
+
+	if (InPinType.PinCategory == UEdGraphSchema_K2::PC_Boolean)
+	{
+		PropertyReference.ReferenceType = ESceneStatePropertyReferenceType::Bool;
+	}
+	else if (InPinType.PinCategory == UEdGraphSchema_K2::PC_Byte)
+	{
+		if (UEnum* Enum = Cast<UEnum>(InPinType.PinSubCategoryObject))
+		{
+			PropertyReference.ReferenceType = ESceneStatePropertyReferenceType::Enum;
+			PropertyReference.TypeObject = InPinType.PinSubCategoryObject.Get();
+		}
+		else
+		{
+			PropertyReference.ReferenceType = ESceneStatePropertyReferenceType::Byte;
+		}
+	}
+	else if (InPinType.PinCategory == UEdGraphSchema_K2::PC_Int)
+	{
+		PropertyReference.ReferenceType = ESceneStatePropertyReferenceType::Int32;
+	}
+	else if (InPinType.PinCategory == UEdGraphSchema_K2::PC_Int64)
+	{
+		PropertyReference.ReferenceType = ESceneStatePropertyReferenceType::Int64;
+	}
+	else if (InPinType.PinCategory == UEdGraphSchema_K2::PC_Real)
+	{
+		if (InPinType.PinSubCategory == UEdGraphSchema_K2::PC_Float)
+		{
+			PropertyReference.ReferenceType = ESceneStatePropertyReferenceType::Float;
+		}
+		else if (InPinType.PinSubCategory == UEdGraphSchema_K2::PC_Double)
+		{
+			PropertyReference.ReferenceType = ESceneStatePropertyReferenceType::Double;
+		}		
+	}
+	else if (InPinType.PinCategory == UEdGraphSchema_K2::PC_Name)
+	{
+		PropertyReference.ReferenceType = ESceneStatePropertyReferenceType::Name;
+	}
+	else if (InPinType.PinCategory == UEdGraphSchema_K2::PC_String)
+	{
+		PropertyReference.ReferenceType = ESceneStatePropertyReferenceType::String;
+	}
+	else if (InPinType.PinCategory == UEdGraphSchema_K2::PC_Text)
+	{
+		PropertyReference.ReferenceType = ESceneStatePropertyReferenceType::Text;
+	}
+	else if (InPinType.PinCategory == UEdGraphSchema_K2::PC_Enum)
+	{
+		PropertyReference.ReferenceType = ESceneStatePropertyReferenceType::Enum;
+		PropertyReference.TypeObject = InPinType.PinSubCategoryObject.Get();
+	}
+	else if (InPinType.PinCategory == UEdGraphSchema_K2::PC_Struct)
+	{
+		PropertyReference.ReferenceType = ESceneStatePropertyReferenceType::Struct;
+		PropertyReference.TypeObject = InPinType.PinSubCategoryObject.Get();
+	}
+	else if (InPinType.PinCategory == UEdGraphSchema_K2::PC_Object)
+	{
+		PropertyReference.ReferenceType = ESceneStatePropertyReferenceType::Object;
+		PropertyReference.TypeObject = InPinType.PinSubCategoryObject.Get();
+	}
+	else if (InPinType.PinCategory == UEdGraphSchema_K2::PC_SoftObject)
+	{
+		PropertyReference.ReferenceType = ESceneStatePropertyReferenceType::SoftObject;
+		PropertyReference.TypeObject = InPinType.PinSubCategoryObject.Get();
+	}
+	else if (InPinType.PinCategory == UEdGraphSchema_K2::PC_Class)
+	{
+		PropertyReference.ReferenceType = ESceneStatePropertyReferenceType::Class;
+		PropertyReference.TypeObject = InPinType.PinSubCategoryObject.Get();
+	}
+	else if (InPinType.PinCategory == UEdGraphSchema_K2::PC_SoftClass)
+	{
+		PropertyReference.ReferenceType = ESceneStatePropertyReferenceType::SoftClass;
+		PropertyReference.TypeObject = InPinType.PinSubCategoryObject.Get();
+	}
+	else
+	{
+		ensureMsgf(false, TEXT("Unhandled pin category %s"), *InPinType.PinCategory.ToString());
+	}
+
+	return PropertyReference;
+}
+#endif
+
+FName FSceneStateBlueprintPropertyReference::GetReferenceTypeMemberName()
+{
+	return GET_MEMBER_NAME_CHECKED(FSceneStateBlueprintPropertyReference, ReferenceType);
+}
+
+FName FSceneStateBlueprintPropertyReference::GetIsReferenceToArrayMemberName()
+{
+	return GET_MEMBER_NAME_CHECKED(FSceneStateBlueprintPropertyReference, bIsReferenceToArray);
+}
+
+FName FSceneStateBlueprintPropertyReference::GetTypeObjectMemberName()
+{
+	return GET_MEMBER_NAME_CHECKED(FSceneStateBlueprintPropertyReference, TypeObject);
+}
diff --git a/Engine/Plugins/Experimental/SceneState/Source/SceneStateBinding/Private/SceneStatePropertyReferenceUtils.cpp b/Engine/Plugins/Experimental/SceneState/Source/SceneStateBinding/Private/SceneStatePropertyReferenceUtils.cpp
--- a/Engine/Plugins/Experimental/SceneState/Source/SceneStateBinding/Private/SceneStatePropertyReferenceUtils.cpp
+++ b/Engine/Plugins/Experimental/SceneState/Source/SceneStateBinding/Private/SceneStatePropertyReferenceUtils.cpp
@@ -1,6 +1,7 @@
 // Copyright Epic Games, Inc. All Rights Reserved.
 
 #include "SceneStatePropertyReferenceUtils.h"
+#include "SceneStateBlueprintPropertyReference.h"
 #include "SceneStatePropertyReference.h"
 
 #if WITH_EDITOR
@@ -29,24 +30,68 @@ static const FLazyName TYPE_Name(TEXT("Name"));
 static const FLazyName TYPE_String(TEXT("String"));
 static const FLazyName TYPE_Text(TEXT("Text"));
 
-TArray<FString> GetPropertyRefTypes(const FProperty* InReferenceProperty)
+struct FReferenceTypeInfo
+{
+	TArray<FString> Types;
+	bool bIsRefToArray = false;
+	bool bCanRefToArray = false;
+};
+
+FReferenceTypeInfo GetReferenceTypeInfo(const FProperty* InReferenceProperty)
 {
 	check(IsPropertyReference(InReferenceProperty));
 
 	FString RefType = InReferenceProperty->GetMetaData(Metadata::RefType);
 	RefType.RemoveSpacesInline();
 
-	TArray<FString> TargetTypes;
-	RefType.ParseIntoArray(TargetTypes, TEXT(","), /*bCullEmpty*/true);
-	return TargetTypes;
+	FReferenceTypeInfo ReferenceTypeInfo;
+	RefType.ParseIntoArray(ReferenceTypeInfo.Types, TEXT(","), /*bCullEmpty*/true);
+	ReferenceTypeInfo.bIsRefToArray = InReferenceProperty->HasMetaData(Metadata::IsRefToArray);
+	ReferenceTypeInfo.bCanRefToArray = InReferenceProperty->HasMetaData(Metadata::CanRefToArray);
+	return ReferenceTypeInfo;
+}
+
+FReferenceTypeInfo GetReferenceTypeInfo(const FSceneStateBlueprintPropertyReference& InPropertyReference)
+{
+	const ESceneStatePropertyReferenceType ReferenceType = InPropertyReference.GetReferenceType();
+	if (ReferenceType == ESceneStatePropertyReferenceType::None)
+	{
+		return FReferenceTypeInfo();
+	}
+
+	UEnum* const ReferenceTypeEnum = StaticEnum<ESceneStatePropertyReferenceType>();
+	const int32 NameIndex = ReferenceTypeEnum->GetIndexByValue(static_cast<int64>(InPropertyReference.GetReferenceType()));
+
+	FString RefType;
+	if (ReferenceTypeEnum->HasMetaData(TEXT("ObjectRef"), NameIndex))
+	{
+		UObject* const TypeObject = InPropertyReference.GetTypeObject();
+		RefType = TypeObject ? TypeObject->GetPathName() : FString();
+		// todo: validation that the type object matches the ReferenceType (e.g. enums, etc)
+	}
+	else
+	{
+		RefType = ReferenceTypeEnum->GetMetaData(TEXT("RefType"), NameIndex);
+	}
+
+	if (RefType.IsEmpty())
+	{
+		return FReferenceTypeInfo();
+	}
+
+	FReferenceTypeInfo ReferenceTypeInfo;
+	ReferenceTypeInfo.Types.Add(MoveTemp(RefType));
+	ReferenceTypeInfo.bIsRefToArray = InPropertyReference.IsReferenceToArray();
+	ReferenceTypeInfo.bCanRefToArray = false; // No support at the moment for BP Property References to have more than 1 pin type
+	return ReferenceTypeInfo;
 }
-	
-bool ArePropertyReferenceCompatible(const FProperty* InTargetReferenceProperty, const FProperty* InSourceReferenceProperty)
+
+bool ArePropertyReferenceCompatible(const FProperty* InTargetReferenceProperty, const void* InTargetAddress, const FProperty* InSourceReferenceProperty, const void* InSourceAddress)
 {
 	check(IsPropertyReference(InSourceReferenceProperty) && IsPropertyReference(InTargetReferenceProperty));
 
-	const TArray<FEdGraphPinType, TInlineAllocator<1>> SourceRefPins = GetPropertyReferencePinTypes(InSourceReferenceProperty);
-	const TArray<FEdGraphPinType, TInlineAllocator<1>> TargetRefPins = GetPropertyReferencePinTypes(InTargetReferenceProperty);
+	const TArray<FEdGraphPinType, TInlineAllocator<1>> SourceRefPins = GetPropertyReferencePinTypes(InSourceReferenceProperty, InSourceAddress);
+	const TArray<FEdGraphPinType, TInlineAllocator<1>> TargetRefPins = GetPropertyReferencePinTypes(InTargetReferenceProperty, InTargetAddress);
 
 	// NOTE: For now be stringent, and return false for properties that have the same Ref types metadata but in different order.
 	return SourceRefPins == TargetRefPins;
@@ -221,14 +266,26 @@ FEdGraphPinType GetTargetPinType(const FString& InTargetType, EPinContainerType
 }
 
 /** Retrieves the Ed Graph Pin Type for a given target type (in string) and a desired container type */
-TArray<FEdGraphPinType, TInlineAllocator<1>> GetTargetPinTypes(const TArray<FString>& InTargetTypes, EPinContainerType InContainerType)
+TArray<FEdGraphPinType, TInlineAllocator<1>> GetTargetPinTypes(const FReferenceTypeInfo& InReferenceTypeInfo)
 {
 	TArray<FEdGraphPinType, TInlineAllocator<1>> PinTypes;
-	PinTypes.Reserve(InTargetTypes.Num());
+	PinTypes.Reserve(InReferenceTypeInfo.Types.Num());
+
+	const EPinContainerType ContainerType = InReferenceTypeInfo.bIsRefToArray
+		? EPinContainerType::Array
+		: EPinContainerType::None;
 
-	for (const FString& TargetType : InTargetTypes)
+	for (const FString& TargetType : InReferenceTypeInfo.Types)
 	{
-		PinTypes.Add(GetTargetPinType(TargetType, InContainerType));
+		const FEdGraphPinType PinType = GetTargetPinType(TargetType, ContainerType);
+		PinTypes.Add(PinType);
+
+		// If Property Reference supports arrays, add the array counterpart too (if the pin type is not an array/container already)
+		if (!PinType.IsContainer() && InReferenceTypeInfo.bCanRefToArray)
+		{
+			FEdGraphPinType& ContainerPinType = PinTypes.Add_GetRef(PinType);
+			ContainerPinType.ContainerType = EPinContainerType::Array;
+		}
 	}
 
 	return PinTypes;
@@ -244,7 +301,7 @@ bool IsPropertyReference(const FProperty* InProperty)
 }
 
 #if WITH_EDITOR
-bool IsPropertyReferenceCompatible(const FProperty* InReferenceProperty, const FProperty* InSourceProperty)
+bool IsPropertyReferenceCompatible(const FProperty* InReferenceProperty, const void* InReferenceAddress, const FProperty* InSourceProperty, const void* InSourceAddress)
 {
 	check(InReferenceProperty && InSourceProperty);
 
@@ -255,30 +312,40 @@ bool IsPropertyReferenceCompatible(const FProperty* InReferenceProperty, const F
 
 	if (IsPropertyReference(InSourceProperty))
 	{
-		return Private::ArePropertyReferenceCompatible(InReferenceProperty, InSourceProperty);
+		return Private::ArePropertyReferenceCompatible(InReferenceProperty, InReferenceAddress, InSourceProperty, InSourceAddress);
 	}
 
-	const FProperty* TestProperty = InSourceProperty;
+	const FStructProperty* ReferenceStructProperty = CastFieldChecked<FStructProperty>(InReferenceProperty);
+	check(ReferenceStructProperty->Struct);
+
+	Private::FReferenceTypeInfo ReferenceTypeInfo;
 
-	const bool bCanTargetRefArray = InReferenceProperty->HasMetaData(Metadata::CanRefToArray);
-	const bool bIsTargetRefArray = InReferenceProperty->HasMetaData(Metadata::IsRefToArray);
+	if (ReferenceStructProperty->Struct->IsChildOf<FSceneStateBlueprintPropertyReference>())
+	{
+		const FSceneStateBlueprintPropertyReference& PropertyReference = *static_cast<const FSceneStateBlueprintPropertyReference*>(InReferenceAddress);
+		ReferenceTypeInfo = Private::GetReferenceTypeInfo(PropertyReference);
+	}
+	else
+	{
+		ReferenceTypeInfo = Private::GetReferenceTypeInfo(InReferenceProperty);
+	}
 
-	if (bIsTargetRefArray || bCanTargetRefArray)
+	const FProperty* TestProperty = InSourceProperty;
+	if (ReferenceTypeInfo.bIsRefToArray || ReferenceTypeInfo.bCanRefToArray)
 	{
 		if (const FArrayProperty* ArrayProperty = CastField<FArrayProperty>(TestProperty))
 		{
 			TestProperty = ArrayProperty->Inner;
 		}
-		else if(!bCanTargetRefArray)
+		else if(!ReferenceTypeInfo.bCanRefToArray)
 		{
 			return false;
 		}
 	}
 
-	const TArray<FString> TargetTypes = Private::GetPropertyRefTypes(InReferenceProperty);
-	for (const FString& TargetType : TargetTypes)
+	for (const FString& TargetType : ReferenceTypeInfo.Types)
 	{
-		if (Private::IsPropertyCompatible(TargetType, TestProperty, bIsTargetRefArray))
+		if (Private::IsPropertyCompatible(TargetType, TestProperty, ReferenceTypeInfo.bIsRefToArray))
 		{
 			return true;
 		}
@@ -286,28 +353,36 @@ bool IsPropertyReferenceCompatible(const FProperty* InReferenceProperty, const F
 	return false;
 }
 
-FEdGraphPinType GetPropertyReferencePinType(const FProperty* InProperty)
+FEdGraphPinType GetPropertyReferencePinType(const FProperty* InProperty, const void* InReferenceAddress)
 {
-	const TArray<FEdGraphPinType, TInlineAllocator<1>> PinTypes = GetPropertyReferencePinTypes(InProperty);
-	return !PinTypes.IsEmpty()
-		? PinTypes[0]
+	const TArray<FEdGraphPinType, TInlineAllocator<1>> PropertyReferenceInfo = GetPropertyReferencePinTypes(InProperty, InReferenceAddress);
+	return !PropertyReferenceInfo.IsEmpty()
+		? PropertyReferenceInfo[0]
 		: FEdGraphPinType();
 }
 
-TArray<FEdGraphPinType, TInlineAllocator<1>> GetPropertyReferencePinTypes(const FProperty* InProperty)
+TArray<FEdGraphPinType, TInlineAllocator<1>> GetPropertyReferencePinTypes(const FSceneStateBlueprintPropertyReference& InPropertyReference)
+{
+	return Private::GetTargetPinTypes(Private::GetReferenceTypeInfo(InPropertyReference));
+}
+
+TArray<FEdGraphPinType, TInlineAllocator<1>> GetPropertyReferencePinTypes(const FProperty* InProperty, const void* InReferenceAddress)
 {
-	if (!ensure(InProperty))
+	if (!ensure(InProperty) || !ensure(IsPropertyReference(InProperty)))
 	{
 		return {};
 	}
 
-	ensure(IsPropertyReference(InProperty));
+	const FStructProperty* StructProperty = CastFieldChecked<FStructProperty>(InProperty);
+	check(StructProperty->Struct);
 
-	const EPinContainerType ContainerType = InProperty->HasMetaData(Metadata::IsRefToArray)
-		? EPinContainerType::Array
-		: EPinContainerType::None;
+	if (StructProperty->Struct->IsChildOf<FSceneStateBlueprintPropertyReference>())
+	{
+		const FSceneStateBlueprintPropertyReference& BlueprintPropertyReference = *static_cast<const FSceneStateBlueprintPropertyReference*>(InReferenceAddress);
+		return GetPropertyReferencePinTypes(BlueprintPropertyReference);
+	}
 
-	return Private::GetTargetPinTypes(Private::GetPropertyRefTypes(InProperty), ContainerType);
+	return Private::GetTargetPinTypes(Private::GetReferenceTypeInfo(InProperty));
 }
 #endif
 
diff --git a/Engine/Plugins/Experimental/SceneState/Source/SceneStateBinding/Public/SceneStateBlueprintPropertyReference.h b/Engine/Plugins/Experimental/SceneState/Source/SceneStateBinding/Public/SceneStateBlueprintPropertyReference.h
new file mode 100644
--- /dev/null
+++ b/Engine/Plugins/Experimental/SceneState/Source/SceneStateBinding/Public/SceneStateBlueprintPropertyReference.h
@@ -0,0 +1,90 @@
+// Copyright Epic Games, Inc. All Rights Reserved.
+
+#pragma once
+
+#include "SceneStatePropertyReference.h"
+#include "SceneStateBlueprintPropertyReference.generated.h"
+
+#define UE_API SCENESTATEBINDING_API
+
+#if WITH_EDITOR
+struct FEdGraphPinType;
+#endif
+
+UENUM()
+enum class ESceneStatePropertyReferenceType : uint8
+{
+	None       UMETA(Hidden),
+	Bool       UMETA(RefType="bool"),
+	Byte       UMETA(RefType="byte"),
+	Int32      UMETA(RefType="int32"),
+	Int64      UMETA(RefType="int64"),
+	Float      UMETA(RefType="float"),
+	Double     UMETA(RefType="double"),
+	Name       UMETA(RefType="Name"),
+	String     UMETA(RefType="String"),
+	Text       UMETA(RefType="Text"),
+	Enum       UMETA(ObjectRef),
+	Struct     UMETA(ObjectRef),
+	Object     UMETA(ObjectRef),
+	SoftObject UMETA(ObjectRef),
+	Class      UMETA(ObjectRef),
+	SoftClass  UMETA(ObjectRef),
+};
+
+/** Property Reference as a Blueprint Type. Usable in blueprints (e.g. tasks) and parameters */
+USTRUCT(BlueprintType, DisplayName="Scene State Property Reference")
+struct FSceneStateBlueprintPropertyReference : public FSceneStatePropertyReference
+{
+	GENERATED_BODY()
+
+#if WITH_EDITOR
+	/** Builds a property reference instance based on the given pin type */
+	UE_API static FSceneStateBlueprintPropertyReference BuildFromPinType(const FEdGraphPinType& InPinType);
+#endif
+
+	/** Returns the property reference's type */
+	ESceneStatePropertyReferenceType GetReferenceType() const
+	{
+		return ReferenceType;
+	}
+
+	/** Returns true if referenced property is an array. */
+	bool IsReferenceToArray() const
+	{
+		return bIsReferenceToArray;
+	}
+
+	/** Returns selected ScriptStruct, Class or Enum. */
+	UObject* GetTypeObject() const
+	{
+		return TypeObject;
+	}
+
+	bool operator==(const FSceneStateBlueprintPropertyReference& InOther) const
+	{
+		return ReferenceType == InOther.ReferenceType
+			&& bIsReferenceToArray == InOther.bIsReferenceToArray
+			&& TypeObject == InOther.TypeObject;
+	}
+
+	/** Get Property Member Name functions */
+	UE_API static FName GetReferenceTypeMemberName();
+	UE_API static FName GetIsReferenceToArrayMemberName();
+	UE_API static FName GetTypeObjectMemberName();
+
+private:
+	/** Specifies the type of property to reference */
+	UPROPERTY(EditAnywhere, Category="Reference Type")
+	ESceneStatePropertyReferenceType ReferenceType = ESceneStatePropertyReferenceType::None;
+
+	/** If specified, the reference is to an TArray<RefType> */
+	UPROPERTY(EditAnywhere, Category="Reference Type")
+	bool bIsReferenceToArray = false;
+
+	/** Specifies the type of property to reference together with ReferenceType, used for Enums, Structs, Objects and Classes. */
+	UPROPERTY(EditAnywhere, Category="Reference Type")
+	TObjectPtr<UObject> TypeObject = nullptr;
+};
+
+#undef UE_API
diff --git a/Engine/Plugins/Experimental/SceneState/Source/SceneStateBinding/Public/SceneStatePropertyReferenceUtils.h b/Engine/Plugins/Experimental/SceneState/Source/SceneStateBinding/Public/SceneStatePropertyReferenceUtils.h
--- a/Engine/Plugins/Experimental/SceneState/Source/SceneStateBinding/Public/SceneStatePropertyReferenceUtils.h
+++ b/Engine/Plugins/Experimental/SceneState/Source/SceneStateBinding/Public/SceneStatePropertyReferenceUtils.h
@@ -8,6 +8,7 @@
 #include "UObject/UnrealType.h"
 
 struct FEdGraphPinType;
+struct FSceneStateBlueprintPropertyReference;
 
 namespace UE::SceneState
 {
@@ -16,13 +17,16 @@ namespace UE::SceneState
 
 #if WITH_EDITOR
 	/** Returns whether a source property is compatible with the reference property */
-	SCENESTATEBINDING_API bool IsPropertyReferenceCompatible(const FProperty* InReferenceProperty, const FProperty* InSourceProperty);
+	SCENESTATEBINDING_API bool IsPropertyReferenceCompatible(const FProperty* InReferenceProperty, const void* InReferenceAddress, const FProperty* InSourceProperty, const void* InSourceAddress);
 
 	/** Returns the first graph pin type based on a given Property Reference property's meta-data */
-	SCENESTATEBINDING_API FEdGraphPinType GetPropertyReferencePinType(const FProperty* InProperty);
+	SCENESTATEBINDING_API FEdGraphPinType GetPropertyReferencePinType(const FProperty* InProperty, const void* InReferenceAddress);
+
+	/** Returns the graph pin type based on a given Property Reference */
+	SCENESTATEBINDING_API TArray<FEdGraphPinType, TInlineAllocator<1>> GetPropertyReferencePinTypes(const FSceneStateBlueprintPropertyReference& InPropertyReference);
 
 	/** Returns the graph pin types based on a given Property Reference property's meta-data */
-	SCENESTATEBINDING_API TArray<FEdGraphPinType, TInlineAllocator<1>> GetPropertyReferencePinTypes(const FProperty* InProperty);
+	SCENESTATEBINDING_API TArray<FEdGraphPinType, TInlineAllocator<1>> GetPropertyReferencePinTypes(const FProperty* InProperty, const void* InReferenceAddress);
 #endif
 
 	template<class InDataType>
diff --git a/Engine/Plugins/Experimental/SceneState/Source/SceneStateBlueprintEditor/Private/Compiler/SceneStateBindingCompiler.cpp b/Engine/Plugins/Experimental/SceneState/Source/SceneStateBlueprintEditor/Private/Compiler/SceneStateBindingCompiler.cpp
--- a/Engine/Plugins/Experimental/SceneState/Source/SceneStateBlueprintEditor/Private/Compiler/SceneStateBindingCompiler.cpp
+++ b/Engine/Plugins/Experimental/SceneState/Source/SceneStateBlueprintEditor/Private/Compiler/SceneStateBindingCompiler.cpp
@@ -191,9 +191,10 @@ void FBindingCompiler::AddTaskBindingDescs()
 				FSceneStateBindingDesc BindingDesc;
 				BindingDesc.ID = InBindingDesc.Id;
 				BindingDesc.Name = InBindingDesc.Name;
-				BindingDesc.Struct = InBindingDesc.Struct;
+				BindingDesc.Struct = InBindingDesc.DataView.GetStruct();
 				BindingDesc.DataHandle = FSceneStateBindingDataHandle(ESceneStateDataType::TaskExtension, TaskIndex, InBindingDesc.DataIndex);
 
+				This->AddDataView(BindingDesc.DataHandle, InBindingDesc.DataView);
 				This->AddBindingDesc(MoveTemp(BindingDesc));
 			});
 	}
@@ -404,14 +405,22 @@ void FBindingCompiler::CompileReferences()
 
 	for (const FSceneStateBinding& ReferenceBinding : ReferenceBindings)
 	{
-		const FPropertyBindingDataView TemplateDataView = TemplateDataViewMap[ReferenceBinding.TargetDataHandle];
+		const FPropertyBindingDataView* TemplateDataView = TemplateDataViewMap.Find(ReferenceBinding.TargetDataHandle);
+		if (!TemplateDataView)
+		{
+			const FText ErrorMessageFormat = FText::Format(LOCTEXT("FailedResolvingTargetDataView", "Error finding target data view for binding {0}")
+				, FText::FromString(ReferenceBinding.ToString()));
+
+			Context.MessageLog.Error(*ErrorMessageFormat.ToString());
+			continue;
+		}
 
 		TargetIndirections.Reset();
-		if (!ReferenceBinding.GetTargetPath().ResolveIndirectionsWithValue(TemplateDataView, TargetIndirections, &ErrorMessage))
+		if (!ReferenceBinding.GetTargetPath().ResolveIndirectionsWithValue(*TemplateDataView, TargetIndirections, &ErrorMessage))
 		{
 			const FText ErrorMessageFormat = FText::Format(LOCTEXT("FailedResolvingReferenceIndirection"
 				, "Error resolving path in {0} ({1}): {2}")
-				, FText::FromString(GetNameSafe(TemplateDataView.GetStruct()))
+				, FText::FromString(GetNameSafe(TemplateDataView->GetStruct()))
 				, FText::FromString(ReferenceBinding.ToString())
 				, FText::FromString(ErrorMessage));
 
diff --git a/Engine/Plugins/Experimental/SceneState/Source/SceneStateBlueprintEditor/Private/DetailsView/SceneStateBlueprintPropertyReferenceCustomization.cpp b/Engine/Plugins/Experimental/SceneState/Source/SceneStateBlueprintEditor/Private/DetailsView/SceneStateBlueprintPropertyReferenceCustomization.cpp
new file mode 100644
--- /dev/null
+++ b/Engine/Plugins/Experimental/SceneState/Source/SceneStateBlueprintEditor/Private/DetailsView/SceneStateBlueprintPropertyReferenceCustomization.cpp
@@ -0,0 +1,187 @@
+// Copyright Epic Games, Inc. All Rights Reserved.
+
+#include "SceneStateBlueprintPropertyReferenceCustomization.h"
+#include "DetailWidgetRow.h"
+#include "Nodes/SceneStateMachineTaskNode.h"
+#include "SPinTypeSelector.h"
+#include "SceneStateBlueprintPropertyReference.h"
+#include "SceneStatePropertyReferenceUtils.h"
+#include "ScopedTransaction.h"
+#include "Tasks/SceneStateBlueprintableTask.h"
+#include "Widgets/Text/STextBlock.h"
+
+#define LOCTEXT_NAMESPACE "SceneStateBlueprintPropertyReferenceCustomization"
+
+bool USceneStatePropertyReferenceSchema::SupportsPinTypeContainer(TWeakPtr<const FEdGraphSchemaAction> InSchemaAction, const FEdGraphPinType& InPinType, const EPinContainerType& InContainerType) const
+{
+	return InContainerType == EPinContainerType::None || InContainerType == EPinContainerType::Array;
+}
+
+namespace UE::SceneState::Editor
+{
+
+namespace Private
+{
+
+/** Determines if the outer object of the property handle is a Task Node in a State Machine */
+bool IsInTaskNode(const TSharedRef<IPropertyHandle>& InPropertyHandle)
+{
+	TArray<UObject*> OuterObjects;
+	InPropertyHandle->GetOuterObjects(OuterObjects);
+	for (UObject* OuterObject : OuterObjects)
+	{
+		if (!OuterObject)
+		{
+			continue;
+		}
+		if (OuterObject->IsA<USceneStateMachineTaskNode>())
+		{
+			return true;
+		}
+		// Only allow Blueprintable Tasks to show the edit properties when it's the Template (i.e. in its Task Blueprint Editor)
+		if (OuterObject->IsA<USceneStateBlueprintableTask>() && !OuterObject->IsTemplate())
+		{
+			return true;
+		}
+	}
+	return false;
+}
+
+} // UE::SceneState::Editor::Private
+
+void FBlueprintPropertyReferenceCustomization::CustomizeHeader(TSharedRef<IPropertyHandle> InPropertyHandle, FDetailWidgetRow& InHeaderRow, IPropertyTypeCustomizationUtils& InCustomizationUtils)
+{
+	PropertyHandle = InPropertyHandle;
+
+	// Hide value if the property handle is being shown in a Task node.
+	// If there's ever a need to edit reference type for a given reference property, a meta-data (e.g. TaskEditable) can be added to address that specific situation. 
+	if (Private::IsInTaskNode(InPropertyHandle))
+	{
+		InHeaderRow
+			.NameContent()
+			[
+				InPropertyHandle->CreatePropertyNameWidget()
+			];
+		return;
+	}
+
+	InHeaderRow
+		.NameContent()
+		[
+			InPropertyHandle->CreatePropertyNameWidget()
+		]
+		.ValueContent()
+		[
+			SNew(SHorizontalBox)
+			+ SHorizontalBox::Slot()
+			.VAlign(VAlign_Center)
+			.AutoWidth()
+			.Padding(FMargin(0, 0, 6, 0))
+			[
+				SNew(STextBlock)
+				.Text(LOCTEXT("ReferenceTo", "Reference to"))
+				.Font(IPropertyTypeCustomizationUtils::GetRegularFont())
+				.ColorAndOpacity(FSlateColor::UseSubduedForeground())
+				.ToolTipText(LOCTEXT("ReferenceToTooltip", "Specifies the type of the referenced property. The referenced property is bound using property binding in Scene State."))
+			]
+			+ SHorizontalBox::Slot()
+			.AutoWidth()
+			[
+				SNew(SPinTypeSelector, FGetPinTypeTree::CreateUObject(GetDefault<USceneStatePropertyReferenceSchema>(), &USceneStatePropertyReferenceSchema::GetVariableTypeTree))
+				.TargetPinType(this, &FBlueprintPropertyReferenceCustomization::GetPinType)
+				.OnPinTypeChanged(this, &FBlueprintPropertyReferenceCustomization::OnPinTypeChanged)
+				.Schema(GetDefault<USceneStatePropertyReferenceSchema>())
+				.Font(IPropertyTypeCustomizationUtils::GetRegularFont())
+				.bAllowArrays(true)
+			]
+		];
+}
+
+void FBlueprintPropertyReferenceCustomization::CustomizeChildren(TSharedRef<IPropertyHandle> InPropertyHandle, IDetailChildrenBuilder& InChildBuilder, IPropertyTypeCustomizationUtils& InCustomizationUtils)
+{
+}
+
+bool FBlueprintPropertyReferenceCustomization::GetPropertyReference(FSceneStateBlueprintPropertyReference& OutPropertyReference) const
+{
+	if (!PropertyHandle.IsValid())
+	{
+		return false;
+	}
+
+	TOptional<FSceneStateBlueprintPropertyReference> Result;
+
+	PropertyHandle->EnumerateConstRawData(
+		[&Result](const void* InRawData, const int32, const int32)->bool
+		{
+			if (!InRawData)
+			{
+				Result.Reset();
+				return false; // break
+			}
+
+			const FSceneStateBlueprintPropertyReference& PropertyReference = *static_cast<const FSceneStateBlueprintPropertyReference*>(InRawData);
+			if (Result.IsSet())
+			{
+				if (*Result != PropertyReference)
+				{
+					Result.Reset();
+					return false; // break, multiple values
+				}
+			}
+			else
+			{
+				Result = PropertyReference;
+			}
+			return true; // continue
+		});
+
+	if (Result.IsSet())
+	{
+		OutPropertyReference = *Result;
+		return true;
+	}
+	return false;
+}
+
+FEdGraphPinType FBlueprintPropertyReferenceCustomization::GetPinType() const
+{
+	FSceneStateBlueprintPropertyReference PropertyReference;
+	if (!GetPropertyReference(PropertyReference))
+	{
+		return FEdGraphPinType();
+	}
+
+	const TArray<FEdGraphPinType, TInlineAllocator<1>> PinTypes = UE::SceneState::GetPropertyReferencePinTypes(PropertyReference);
+	return !PinTypes.IsEmpty()
+		? PinTypes[0]
+		: FEdGraphPinType();
+}
+
+void FBlueprintPropertyReferenceCustomization::OnPinTypeChanged(const FEdGraphPinType& InPinType) const
+{
+	if (!PropertyHandle.IsValid())
+	{
+		return;
+	}
+
+	FScopedTransaction Transaction(LOCTEXT("SetPropertyReferenceType", "Set Property Reference Type"));
+
+	const FSceneStateBlueprintPropertyReference NewValue = FSceneStateBlueprintPropertyReference::BuildFromPinType(InPinType);
+
+	PropertyHandle->NotifyPreChange();
+	PropertyHandle->EnumerateRawData(
+		[&NewValue](void* InRawData, const int32, const int32)->bool
+		{
+			if (InRawData)
+			{
+				*static_cast<FSceneStateBlueprintPropertyReference*>(InRawData) = NewValue;
+			}
+			return true; // continue
+		});
+	PropertyHandle->NotifyPostChange(EPropertyChangeType::ValueSet);
+	PropertyHandle->NotifyFinishedChangingProperties();
+}
+
+} // UE::SceneState::Editor
+
+#undef LOCTEXT_NAMESPACE
diff --git a/Engine/Plugins/Experimental/SceneState/Source/SceneStateBlueprintEditor/Private/DetailsView/SceneStateBlueprintPropertyReferenceCustomization.h b/Engine/Plugins/Experimental/SceneState/Source/SceneStateBlueprintEditor/Private/DetailsView/SceneStateBlueprintPropertyReferenceCustomization.h
new file mode 100644
--- /dev/null
+++ b/Engine/Plugins/Experimental/SceneState/Source/SceneStateBlueprintEditor/Private/DetailsView/SceneStateBlueprintPropertyReferenceCustomization.h
@@ -0,0 +1,50 @@
+// Copyright Epic Games, Inc. All Rights Reserved.
+
+#pragma once
+
+#include "EdGraphSchema_K2.h"
+#include "IPropertyTypeCustomization.h"
+#include "SceneStateBlueprintPropertyReferenceCustomization.generated.h"
+
+struct FSceneStateBlueprintPropertyReference;
+
+/** Property Reference Schema to allow customizing the requirements (e.g. supported containers). */
+UCLASS()
+class USceneStatePropertyReferenceSchema : public UEdGraphSchema_K2
+{
+	GENERATED_BODY()
+
+public:
+	//~ Begin UEdGraphSchema_K2
+	virtual bool SupportsPinTypeContainer(TWeakPtr<const FEdGraphSchemaAction> InSchemaAction, const FEdGraphPinType& InPinType, const EPinContainerType& InContainerType) const override;
+	//~ End UEdGraphSchema_K2
+};
+
+namespace UE::SceneState::Editor
+{
+
+/** Customization for FSceneStateBlueprintPropertyReference */
+class FBlueprintPropertyReferenceCustomization : public IPropertyTypeCustomization
+{
+public:
+	static TSharedRef<IPropertyTypeCustomization> MakeInstance()
+	{
+		return MakeShared<FBlueprintPropertyReferenceCustomization>();
+	}
+
+	//~ Begin IPropertyTypeCustomization
+	virtual void CustomizeHeader(TSharedRef<IPropertyHandle> InPropertyHandle, FDetailWidgetRow& InHeaderRow, IPropertyTypeCustomizationUtils& InCustomizationUtils) override;
+	virtual void CustomizeChildren(TSharedRef<IPropertyHandle> InPropertyHandle, IDetailChildrenBuilder& InChildBuilder, IPropertyTypeCustomizationUtils& InCustomizationUtils) override;
+	//~ End IPropertyTypeCustomization
+
+private:
+	bool GetPropertyReference(FSceneStateBlueprintPropertyReference& OutPropertyReference) const;
+
+	FEdGraphPinType GetPinType() const;
+
+	void OnPinTypeChanged(const FEdGraphPinType& InPinType) const;
+
+	TSharedPtr<IPropertyHandle> PropertyHandle;
+};
+
+} // UE::SceneState::Editor
diff --git a/Engine/Plugins/Experimental/SceneState/Source/SceneStateBlueprintEditor/Private/DetailsView/SceneStateCachedBindingData.cpp b/Engine/Plugins/Experimental/SceneState/Source/SceneStateBlueprintEditor/Private/DetailsView/SceneStateCachedBindingData.cpp
--- a/Engine/Plugins/Experimental/SceneState/Source/SceneStateBlueprintEditor/Private/DetailsView/SceneStateCachedBindingData.cpp
+++ b/Engine/Plugins/Experimental/SceneState/Source/SceneStateBlueprintEditor/Private/DetailsView/SceneStateCachedBindingData.cpp
@@ -5,6 +5,7 @@
 #include "IPropertyAccessEditor.h"
 #include "PropertyBindingDataView.h"
 #include "SceneStateBindingDesc.h"
+#include "SceneStateBlueprintPropertyReference.h"
 #include "SceneStatePropertyReference.h"
 #include "SceneStatePropertyReferenceMetadata.h"
 #include "SceneStatePropertyReferenceUtils.h"
@@ -35,9 +36,14 @@ void FSceneStateCachedBindingData::UpdatePropertyReferenceTooltip(const FPropert
 
 bool FSceneStateCachedBindingData::GetPinTypeAndIconForProperty(const FProperty& InProperty, FPropertyBindingDataView InTargetDataView, FEdGraphPinType& OutPinType, FName& OutIconName) const
 {
-	if (UE::SceneState::IsPropertyReference(&InProperty))
+	if (UE::SceneState::IsPropertyReference(&InProperty) && InTargetDataView.IsValid())
 	{
-		OutPinType = GetPropertyReferencePinType(&InProperty);
+		// Use internal type to construct PinType if it's property of PropertyRef type.
+		TArray<FPropertyBindingPathIndirection> TargetIndirections;
+		if (ensure(GetTargetPath().ResolveIndirectionsWithValue(InTargetDataView, TargetIndirections)))
+		{
+			OutPinType = GetPropertyReferencePinType(&InProperty, TargetIndirections.Last().GetPropertyAddress());
+		}
 		OutIconName = TEXT("Kismet.Tabs.Variables");
 		return true;
 	}
@@ -60,21 +66,19 @@ void FSceneStateCachedBindingData::AddPropertyInfoOverride(const FProperty& InPr
 	// Add the PropertyRef property type with its RefType
 	if (StructProperty->Struct->IsChildOf<FSceneStatePropertyReference>())
 	{
-		TArray<FEdGraphPinType, TInlineAllocator<1>> PinTypes = GetPropertyReferencePinTypes(&InProperty);
-
-		// If Property supports Arrays, add the Array version of these pin types
-		const bool bCanRefToArray = GetPropertyHandle()->HasMetaData(Metadata::CanRefToArray);
-
-		OutPropertyInfoOverrides.Reserve(OutPropertyInfoOverrides.Num() + PinTypes.Num());
+		void* PropertyRefAddress = nullptr;
 
-		for (FEdGraphPinType& PinType : PinTypes)
+		// Passed in InProperty comes from GetPropertyHandle
+		if (const IPropertyHandle* PropertyRefHandle = GetPropertyHandle())
 		{
-			if (bCanRefToArray)
-			{
-				PinType.ContainerType = EPinContainerType::Array;
-			}
+			PropertyRefHandle->GetValueData(PropertyRefAddress);
+		}
 
+		TArray<FEdGraphPinType, TInlineAllocator<1>> PropertyReferenceInfo = GetPropertyReferencePinTypes(&InProperty, PropertyRefAddress);
+		OutPropertyInfoOverrides.Reserve(OutPropertyInfoOverrides.Num() + PropertyReferenceInfo.Num());
 
+		for (const FEdGraphPinType& PinType : PropertyReferenceInfo)
+		{
 			TSharedRef<UE::PropertyBinding::FPropertyInfoOverride> PropertyInfoOverride = MakeShared<UE::PropertyBinding::FPropertyInfoOverride>();
 			PropertyInfoOverride->PinType = PinType;
 
@@ -98,7 +102,7 @@ bool FSceneStateCachedBindingData::DeterminePropertiesCompatibilityInternal(cons
 {
 	if (UE::SceneState::IsPropertyReference(InTargetProperty))
 	{
-		bOutAreCompatible = UE::SceneState::IsPropertyReferenceCompatible(InTargetProperty, InSourceProperty);
+		bOutAreCompatible = UE::SceneState::IsPropertyReferenceCompatible(InTargetProperty, InTargetPropertyValue, InSourceProperty, InSourcePropertyValue);
 		return true;
 	}
 	return false;
diff --git a/Engine/Plugins/Experimental/SceneState/Source/SceneStateBlueprintEditor/Private/SceneStateBlueprintEditorModule.cpp b/Engine/Plugins/Experimental/SceneState/Source/SceneStateBlueprintEditor/Private/SceneStateBlueprintEditorModule.cpp
--- a/Engine/Plugins/Experimental/SceneState/Source/SceneStateBlueprintEditor/Private/SceneStateBlueprintEditorModule.cpp
+++ b/Engine/Plugins/Experimental/SceneState/Source/SceneStateBlueprintEditor/Private/SceneStateBlueprintEditorModule.cpp
@@ -4,6 +4,7 @@
 #include "AssetToolsModule.h"
 #include "BlueprintEditorModule.h"
 #include "Compiler/SceneStateBlueprintCompilerContext.h"
+#include "DetailsView/SceneStateBlueprintPropertyReferenceCustomization.h"
 #include "DetailsView/SceneStateBlueprintableTaskInstanceCustomization.h"
 #include "DetailsView/SceneStateMachineGraphCustomization.h"
 #include "DetailsView/SceneStateMachineTaskInstanceCustomization.h"
@@ -18,6 +19,7 @@
 #include "SceneStateBlueprint.h"
 #include "SceneStateBlueprintEditorCommands.h"
 #include "SceneStateBlueprintEditorLog.h"
+#include "SceneStateBlueprintPropertyReference.h"
 #include "SceneStateMachineGraph.h"
 #include "SceneStateMachineGraphSchema.h"
 #include "SceneStateObject.h"
@@ -167,6 +169,9 @@ void FBlueprintEditorModule::RegisterDetailCustomizations()
 	PropertyEditorModule.RegisterCustomPropertyTypeLayout(CustomizedTypes.Add_GetRef(FSceneStatePlayerTaskInstance::StaticStruct()->GetFName())
 		, FOnGetPropertyTypeCustomizationInstance::CreateStatic(&FPlayerTaskInstanceCustomization::MakeInstance));
 
+	PropertyEditorModule.RegisterCustomPropertyTypeLayout(CustomizedTypes.Add_GetRef(FSceneStateBlueprintPropertyReference::StaticStruct()->GetFName())
+		, FOnGetPropertyTypeCustomizationInstance::CreateStatic(&FBlueprintPropertyReferenceCustomization::MakeInstance));
+
 	PropertyEditorModule.RegisterCustomClassLayout(CustomizedClasses.Add_GetRef(USceneStateMachineGraph::StaticClass()->GetFName())
 		, FOnGetDetailCustomizationInstance::CreateStatic(&FStateMachineGraphCustomization::MakeInstance));
 
diff --git a/Engine/Plugins/Experimental/SceneState/Source/SceneStateBlueprintEditor/SceneStateBlueprintEditor.Build.cs b/Engine/Plugins/Experimental/SceneState/Source/SceneStateBlueprintEditor/SceneStateBlueprintEditor.Build.cs
--- a/Engine/Plugins/Experimental/SceneState/Source/SceneStateBlueprintEditor/SceneStateBlueprintEditor.Build.cs
+++ b/Engine/Plugins/Experimental/SceneState/Source/SceneStateBlueprintEditor/SceneStateBlueprintEditor.Build.cs
@@ -29,6 +29,7 @@ public class SceneStateBlueprintEditor : ModuleRules
                 "InputCore",
                 "Kismet",
                 "KismetCompiler",
+                "KismetWidgets",
                 "PropertyBindingUtils",
                 "PropertyBindingUtilsEditor",
                 "PropertyEditor",
