Motion Design: fix crash with scene state when task nodes have prerequisite nodes that aren't connected to the owning state node.

diff --git a/Engine/Plugins/Experimental/SceneState/Source/SceneStateMachineEditor/Private/SceneStateMachineTaskCompiler.cpp b/Engine/Plugins/Experimental/SceneState/Source/SceneStateMachineEditor/Private/SceneStateMachineTaskCompiler.cpp
--- a/Engine/Plugins/Experimental/SceneState/Source/SceneStateMachineEditor/Private/SceneStateMachineTaskCompiler.cpp
+++ b/Engine/Plugins/Experimental/SceneState/Source/SceneStateMachineEditor/Private/SceneStateMachineTaskCompiler.cpp
@@ -146,11 +146,15 @@ void FStateMachineTaskCompiler::CompileTasks(FCompileResult& OutCompilationResul
 		OutCompilationResult.TaskMetadata.Add(TaskInfo.Metadata);
 
 		Task.PrerequisiteRange.Index = OutCompilationResult.TaskPrerequisites.Num();
-		Task.PrerequisiteRange.Count = TaskInfo.Prerequisites.Num();
+		Task.PrerequisiteRange.Count = 0;
 
 		for (const USceneStateMachineTaskNode* PrerequisiteNode : TaskInfo.Prerequisites)
 		{
-			OutCompilationResult.TaskPrerequisites.Add(ProcessedNodes[PrerequisiteNode]);
+			if (const uint16* PrerequisiteNodeIndex = ProcessedNodes.Find(PrerequisiteNode))
+			{
+				OutCompilationResult.TaskPrerequisites.Add(*PrerequisiteNodeIndex);
+				++Task.PrerequisiteRange.Count;
+			}
 		}
 	}
 
