Motion Design: Scene State user defined struct variables to show friendly names instead of underlying property name

diff --git a/Engine/Plugins/Experimental/SceneState/Source/SceneStateBlueprintEditor/Private/DetailsView/SceneStateCachedBindingData.cpp b/Engine/Plugins/Experimental/SceneState/Source/SceneStateBlueprintEditor/Private/DetailsView/SceneStateCachedBindingData.cpp
--- a/Engine/Plugins/Experimental/SceneState/Source/SceneStateBlueprintEditor/Private/DetailsView/SceneStateCachedBindingData.cpp
+++ b/Engine/Plugins/Experimental/SceneState/Source/SceneStateBlueprintEditor/Private/DetailsView/SceneStateCachedBindingData.cpp
@@ -3,12 +3,16 @@
 #include "SceneStateCachedBindingData.h"
 #include "EdGraphSchema_K2.h"
 #include "IPropertyAccessEditor.h"
+#include "Kismet2/StructureEditorUtils.h"
+#include "Misc/EnumerateRange.h"
+#include "PropertyBindingBindableStructDescriptor.h"
 #include "PropertyBindingDataView.h"
 #include "SceneStateBindingDesc.h"
 #include "SceneStateBlueprintPropertyReference.h"
 #include "SceneStatePropertyReference.h"
 #include "SceneStatePropertyReferenceMetadata.h"
 #include "SceneStatePropertyReferenceUtils.h"
+#include "StructUtils/UserDefinedStruct.h"
 
 #define LOCTEXT_NAMESPACE "SceneStateCachedBindingData"
 
@@ -55,6 +59,67 @@ bool FSceneStateCachedBindingData::IsPropertyReference(const FProperty& InProper
 	return UE::SceneState::IsPropertyReference(&InProperty);
 }
 
+void FSceneStateCachedBindingData::UpdateSourcePropertyPath(TConstStructView<FPropertyBindingBindableStructDescriptor> InDescriptor, const FPropertyBindingPath& InSourcePath, FString& OutString)
+{
+	const FPropertyBindingBindableStructDescriptor* Descriptor = InDescriptor.GetPtr();
+	if (!Descriptor || !Descriptor->Struct)
+	{
+		return;
+	}
+
+	TArray<FPropertyBindingPathIndirection> Indirections;
+	if (!InSourcePath.ResolveIndirections(Descriptor->Struct, Indirections))
+	{
+		return;
+	}
+
+	FStringBuilderBase Result;
+	for (TEnumerateRef<const FPropertyBindingPathSegment> Segment : EnumerateRange(InSourcePath.GetSegments()))
+	{
+		const int32 SegmentIndex = Segment.GetIndex();
+
+		const FPropertyBindingPathIndirection* Indirection = Indirections.FindByPredicate(
+			[SegmentIndex](const FPropertyBindingPathIndirection& InIndirection)
+			{
+				return InIndirection.GetPathSegmentIndex() == SegmentIndex;
+			});
+
+		if (!ensure(Indirection))
+		{
+			return;
+		}
+
+		if (SegmentIndex > 0)
+		{
+			Result += TEXT(".");
+		}
+
+		if (const UUserDefinedStruct* CurrentUserDefinedStruct = Cast<const UUserDefinedStruct>(Indirection->GetContainerStruct()))
+		{
+			const FString FriendlyName = FStructureEditorUtils::GetVariableFriendlyNameForProperty(CurrentUserDefinedStruct, Indirection->GetProperty());
+			if (!FriendlyName.IsEmpty())
+			{
+				Result += FriendlyName;
+			}
+			else
+			{
+				Result += Segment->GetName().ToString();
+			}
+		}
+		else
+		{
+			Result += Segment->GetName().ToString();
+		}
+
+		if (Segment->GetArrayIndex() >= 0)
+		{
+			Result += FString::Printf(TEXT("[%d]"), Segment->GetArrayIndex());
+		}
+	}
+
+	OutString = Result.ToString();
+}
+
 void FSceneStateCachedBindingData::AddPropertyInfoOverride(const FProperty& InProperty, TArray<TSharedPtr<const UE::PropertyBinding::FPropertyInfoOverride>>& OutPropertyInfoOverrides) const
 {
 	const FStructProperty* StructProperty = CastField<const FStructProperty>(&InProperty);
diff --git a/Engine/Plugins/Experimental/SceneState/Source/SceneStateBlueprintEditor/Private/DetailsView/SceneStateCachedBindingData.h b/Engine/Plugins/Experimental/SceneState/Source/SceneStateBlueprintEditor/Private/DetailsView/SceneStateCachedBindingData.h
--- a/Engine/Plugins/Experimental/SceneState/Source/SceneStateBlueprintEditor/Private/DetailsView/SceneStateCachedBindingData.h
+++ b/Engine/Plugins/Experimental/SceneState/Source/SceneStateBlueprintEditor/Private/DetailsView/SceneStateCachedBindingData.h
@@ -16,6 +16,7 @@ struct FSceneStateCachedBindingData : UE::PropertyBinding::FCachedBindingData, T
 	virtual void UpdatePropertyReferenceTooltip(const FProperty& InProperty, FTextBuilder& InOutTextBuilder) const override;
 	virtual bool GetPinTypeAndIconForProperty(const FProperty& InProperty, FPropertyBindingDataView InTargetDataView, FEdGraphPinType& OutPinType, FName& OutIconName) const override;
 	virtual bool IsPropertyReference(const FProperty& InProperty) override;
+	virtual void UpdateSourcePropertyPath(TConstStructView<FPropertyBindingBindableStructDescriptor> InDescriptor, const FPropertyBindingPath& InSourcePath, FString& OutString);
 	virtual void AddPropertyInfoOverride(const FProperty& InProperty, TArray<TSharedPtr<const UE::PropertyBinding::FPropertyInfoOverride>>& OutPropertyInfoOverrides) const override;
 	virtual bool DeterminePropertiesCompatibilityInternal(const FProperty* InSourceProperty,const FProperty* InTargetProperty,const void* InSourcePropertyValue,const void* InTargetPropertyValue,bool& bOutAreCompatible) const override;
 	//~ End UE::PropertyBinding::FCachedBindingData
