Motion Design Scene State: added NoBindingSelfOnly and NoBindingContainerSelfOnly which disallow binding a property but not necessarily its children.
Also introduced 'AnyStruct' for a property reference to not have to specify a particular struct.

diff --git a/Engine/Plugins/Experimental/SceneState/Source/SceneStateBinding/Private/SceneStatePropertyReferenceUtils.cpp b/Engine/Plugins/Experimental/SceneState/Source/SceneStateBinding/Private/SceneStatePropertyReferenceUtils.cpp
--- a/Engine/Plugins/Experimental/SceneState/Source/SceneStateBinding/Private/SceneStatePropertyReferenceUtils.cpp
+++ b/Engine/Plugins/Experimental/SceneState/Source/SceneStateBinding/Private/SceneStatePropertyReferenceUtils.cpp
@@ -29,6 +29,7 @@ static const FLazyName TYPE_Double(TEXT("double"));
 static const FLazyName TYPE_Name(TEXT("Name"));
 static const FLazyName TYPE_String(TEXT("String"));
 static const FLazyName TYPE_Text(TEXT("Text"));
+static const FLazyName TYPE_AnyStruct(TEXT("AnyStruct"));
 
 struct FReferenceTypeInfo
 {
@@ -105,6 +106,10 @@ bool IsPropertyCompatible(const FString& InTargetType, const FProperty* InTestPr
 	{
 		return true;
 	}
+	if (TargetTypeName == TYPE_AnyStruct)
+	{
+		return InTestProperty->IsA<FStructProperty>();
+	}
 	if (TargetTypeName == TYPE_Bool)
 	{
 		return InTestProperty->IsA<FBoolProperty>();
@@ -190,6 +195,10 @@ FEdGraphPinType GetTargetPinType(const FString& InTargetType, EPinContainerType
 	{
 		PinType.PinCategory = UEdGraphSchema_K2::PC_Wildcard;
 	}
+	else if (TargetTypeName == Private::TYPE_AnyStruct)
+	{
+		PinType.PinCategory = UEdGraphSchema_K2::PC_Wildcard;
+	}
 	else if (TargetTypeName == Private::TYPE_Bool)
 	{
 		PinType.PinCategory = UEdGraphSchema_K2::PC_Boolean;
diff --git a/Engine/Plugins/Experimental/SceneState/Source/SceneStateBinding/Public/SceneStateBindingMetadata.h b/Engine/Plugins/Experimental/SceneState/Source/SceneStateBinding/Public/SceneStateBindingMetadata.h
new file mode 100644
--- /dev/null
+++ b/Engine/Plugins/Experimental/SceneState/Source/SceneStateBinding/Public/SceneStateBindingMetadata.h
@@ -0,0 +1,24 @@
+// Copyright Epic Games, Inc. All Rights Reserved.
+
+#pragma once
+
+#include "UObject/NameTypes.h"
+
+namespace UE::SceneState::Metadata
+{
+#if WITH_EDITOR
+	/**
+	 * Describes that the property does not allow binding. But this no binding policy does not propagate to its child properties
+	 * For containers, this policy would apply for both the container and inner properties.
+	 */
+	const FLazyName NoBindingSelfOnly(TEXT("NoBindingSelfOnly"));
+
+	/**
+	 * Version of 'NoBindingSelfOnly' where it applies the no binding policy to the container property only, and not its inner properties
+	 * E.g. for a TArray<FVector> property:
+	 * If 'NoBindingSelfOnly' is used, both the array as a whole and its FVector inner property would be unbindable and only the FVector's individual components (X, Y, Z) would be bindable.
+	 * If 'NoBindingContainerSelfOnly' is used, the array would still be non-bindable, but the FVector inner property would now be bindable.
+	 */
+	const FLazyName NoBindingContainerSelfOnly(TEXT("NoBindingContainerSelfOnly"));
+#endif
+}
diff --git a/Engine/Plugins/Experimental/SceneState/Source/SceneStateBlueprintEditor/Private/DetailsView/SceneStatePropertyBindingExtension.cpp b/Engine/Plugins/Experimental/SceneState/Source/SceneStateBlueprintEditor/Private/DetailsView/SceneStatePropertyBindingExtension.cpp
--- a/Engine/Plugins/Experimental/SceneState/Source/SceneStateBlueprintEditor/Private/DetailsView/SceneStatePropertyBindingExtension.cpp
+++ b/Engine/Plugins/Experimental/SceneState/Source/SceneStateBlueprintEditor/Private/DetailsView/SceneStatePropertyBindingExtension.cpp
@@ -2,6 +2,7 @@
 
 #include "SceneStatePropertyBindingExtension.h"
 #include "GameFramework/Actor.h"
+#include "SceneStateBindingMetadata.h"
 #include "SceneStateBlueprintEditorUtils.h"
 #include "SceneStateCachedBindingData.h"
 #include "SceneStatePropertyReference.h"
@@ -21,6 +22,29 @@ void FBindingExtension::ExtendWidgetRow(FDetailWidgetRow& InWidgetRow, const IDe
 	FPropertyBindingExtension::ExtendWidgetRow(InWidgetRow, InDetailBuilder, InObjectClass, InPropertyHandle);
 }
 
+bool FBindingExtension::IsPropertyExtendable(const UClass* InObjectClass, const IPropertyHandle& InPropertyHandle) const
+{
+	// Disallow if property has 'NoBindingSelfOnly'.
+	// Unlike 'NoBinding', this only checks 'self' and does not recurse check parent handles.
+	if (InPropertyHandle.HasMetaData(Metadata::NoBindingSelfOnly))
+	{
+		return false;
+	}
+
+	// Version of 'NoBindingSelfOnly' for containers to only apply it on the container property, not its inner properties.
+	if (InPropertyHandle.HasMetaData(Metadata::NoBindingContainerSelfOnly))
+	{
+		// Container and its Inner properties carry this same meta-data so check if the property is the container, and not the inner property.
+		const FProperty* Property = InPropertyHandle.GetProperty();
+		if (!Property || Property->IsA<FArrayProperty>() || Property->IsA<FSetProperty>() || Property->IsA<FMapProperty>())
+		{
+			return false;
+		}
+	}
+
+	return FPropertyBindingExtension::IsPropertyExtendable(InObjectClass, InPropertyHandle);
+}
+
 TSharedPtr<UE::PropertyBinding::FCachedBindingData> FBindingExtension::CreateCachedBindingData(IPropertyBindingBindingCollectionOwner* InBindingsOwner
 	, const FPropertyBindingPath& InTargetPath
 	, const TSharedPtr<IPropertyHandle>& InPropertyHandle
diff --git a/Engine/Plugins/Experimental/SceneState/Source/SceneStateBlueprintEditor/Private/DetailsView/SceneStatePropertyBindingExtension.h b/Engine/Plugins/Experimental/SceneState/Source/SceneStateBlueprintEditor/Private/DetailsView/SceneStatePropertyBindingExtension.h
--- a/Engine/Plugins/Experimental/SceneState/Source/SceneStateBlueprintEditor/Private/DetailsView/SceneStatePropertyBindingExtension.h
+++ b/Engine/Plugins/Experimental/SceneState/Source/SceneStateBlueprintEditor/Private/DetailsView/SceneStatePropertyBindingExtension.h
@@ -15,6 +15,7 @@ public:
 	//~ End IDetailPropertyExtensionHandler
 
 	//~ Begin FPropertyBindingExtension
+	virtual bool IsPropertyExtendable(const UClass* InObjectClass, const IPropertyHandle& InPropertyHandle) const override;
 	virtual TSharedPtr<UE::PropertyBinding::FCachedBindingData> CreateCachedBindingData(IPropertyBindingBindingCollectionOwner* InBindingsOwner, const FPropertyBindingPath& InTargetPath, const TSharedPtr<IPropertyHandle>& InPropertyHandle, TConstArrayView<TInstancedStruct<FPropertyBindingBindableStructDescriptor>> InAccessibleStructs) const override;
 	virtual bool GetPromotionToParameterOverrideInternal(const FProperty& InProperty, bool& bOutOverride) const override;
 	virtual bool CanBindToArrayElements() const override;
