Fixed non unity compile errors by adding StrProperty.h

diff --git a/Engine/Plugins/Experimental/SceneState/Source/SceneStateBinding/Public/SceneStatePropertyReferenceUtils.h b/Engine/Plugins/Experimental/SceneState/Source/SceneStateBinding/Public/SceneStatePropertyReferenceUtils.h
--- a/Engine/Plugins/Experimental/SceneState/Source/SceneStateBinding/Public/SceneStatePropertyReferenceUtils.h
+++ b/Engine/Plugins/Experimental/SceneState/Source/SceneStateBinding/Public/SceneStatePropertyReferenceUtils.h
@@ -4,6 +4,7 @@
 
 #include "Concepts/BaseStructureProvider.h"
 #include "UObject/EnumProperty.h"
+#include "UObject/StrProperty.h"
 #include "UObject/TextProperty.h"
 #include "UObject/UnrealType.h"
 
