Motion Design Scene State: changed task UI design to have custom task names and task type

diff --git a/Engine/Plugins/Experimental/SceneState/Source/SceneStateMachineEditor/Private/SceneStateMachineEditorStyle.cpp b/Engine/Plugins/Experimental/SceneState/Source/SceneStateMachineEditor/Private/SceneStateMachineEditorStyle.cpp
--- a/Engine/Plugins/Experimental/SceneState/Source/SceneStateMachineEditor/Private/SceneStateMachineEditorStyle.cpp
+++ b/Engine/Plugins/Experimental/SceneState/Source/SceneStateMachineEditor/Private/SceneStateMachineEditorStyle.cpp
@@ -66,6 +66,10 @@ void FStateMachineEditorStyle::RegisterGraphStyles()
 		.SetShadowColorAndOpacity(FLinearColor(0.f, 0.f, 0.f, 0.f));
 
 	Set("Graph.TaskNode.Title", GraphTaskNodeTitle);
+
+	Set("Graph.Node.NodeTitleExtraLines", FTextBlockStyle(FAppStyle::Get().GetWidgetStyle<FTextBlockStyle>("Graph.Node.NodeTitleExtraLines"))
+		.SetFontSize(FCoreStyle::SmallTextSize));
+
 	Set("Graph.TaskNode.TitleInlineEditableText", FInlineEditableTextBlockStyle(FAppStyle::Get().GetWidgetStyle<FInlineEditableTextBlockStyle>("Graph.StateNode.NodeTitleInlineEditableText"))
 		.SetTextStyle(GraphTaskNodeTitle));
 }
diff --git a/Engine/Plugins/Experimental/SceneState/Source/SceneStateMachineEditor/Private/Widgets/SSceneStateMachineNode.cpp b/Engine/Plugins/Experimental/SceneState/Source/SceneStateMachineEditor/Private/Widgets/SSceneStateMachineNode.cpp
--- a/Engine/Plugins/Experimental/SceneState/Source/SceneStateMachineEditor/Private/Widgets/SSceneStateMachineNode.cpp
+++ b/Engine/Plugins/Experimental/SceneState/Source/SceneStateMachineEditor/Private/Widgets/SSceneStateMachineNode.cpp
@@ -1,8 +1,10 @@
 // Copyright Epic Games, Inc. All Rights Reserved.
 
 #include "SSceneStateMachineNode.h"
+#include "IDocumentation.h"
 #include "SGraphPin.h"
 #include "Widgets/SBoxPanel.h"
+#include "Widgets/SToolTip.h"
 
 namespace UE::SceneState::Editor
 {
@@ -24,6 +26,16 @@ void SStateMachineNode::AddPinWidgetToSlot(const TSharedRef<SGraphPin>& InPinWid
 
 void SStateMachineNode::UpdateGraphNode()
 {
+	if (!SWidget::GetToolTip().IsValid())
+	{
+		TSharedRef<SToolTip> DefaultToolTip = IDocumentation::Get()->CreateToolTip(TAttribute<FText>(this, &SGraphNode::GetNodeTooltip)
+			, /*OverrideContent*/nullptr
+			, GraphNode->GetDocumentationLink()
+			, GraphNode->GetDocumentationExcerptName());
+
+		SetToolTip(DefaultToolTip);
+	}
+
 	InputPins.Empty();
 	OutputPins.Empty();
 
diff --git a/Engine/Plugins/Experimental/SceneState/Source/SceneStateMachineEditor/Private/Widgets/SSceneStateMachineTaskNode.cpp b/Engine/Plugins/Experimental/SceneState/Source/SceneStateMachineEditor/Private/Widgets/SSceneStateMachineTaskNode.cpp
--- a/Engine/Plugins/Experimental/SceneState/Source/SceneStateMachineEditor/Private/Widgets/SSceneStateMachineTaskNode.cpp
+++ b/Engine/Plugins/Experimental/SceneState/Source/SceneStateMachineEditor/Private/Widgets/SSceneStateMachineTaskNode.cpp
@@ -41,7 +41,8 @@ void SStateMachineTaskNode::Construct(const FArguments& InArgs, USceneStateMachi
 
 	NodeTitle = SNew(SNodeTitle, GraphNode)
 		.StyleSet(&Style)
-		.Style(TEXT("Graph.TaskNode.Title"));
+		.Style(TEXT("Graph.TaskNode.Title"))
+		.ExtraLineStyle(TEXT("Graph.Node.NodeTitleExtraLines"));
 
 	InNode->OnPostEditTask().AddSP(NodeTitle.ToSharedRef(), &SNodeTitle::MarkDirty);
 
diff --git a/Engine/Plugins/Experimental/SceneState/Source/SceneStateMachineGraph/Private/Nodes/SceneStateMachineNodeNameValidator.cpp b/Engine/Plugins/Experimental/SceneState/Source/SceneStateMachineGraph/Private/Nodes/SceneStateMachineNodeNameValidator.cpp
--- a/Engine/Plugins/Experimental/SceneState/Source/SceneStateMachineGraph/Private/Nodes/SceneStateMachineNodeNameValidator.cpp
+++ b/Engine/Plugins/Experimental/SceneState/Source/SceneStateMachineGraph/Private/Nodes/SceneStateMachineNodeNameValidator.cpp
@@ -1,27 +1,37 @@
 // Copyright Epic Games, Inc. All Rights Reserved.
 
 #include "SceneStateMachineNodeNameValidator.h"
+#include "Algo/Count.h"
 #include "Nodes/SceneStateMachineNode.h"
 #include "SceneStateMachineGraph.h"
 
 namespace UE::SceneState::Graph
 {
 
-FStateMachineNodeNameValidator::FStateMachineNodeNameValidator(const USceneStateMachineNode* InNode)
+FStateMachineNodeNameValidator::FStateMachineNodeNameValidator(const USceneStateMachineNode* InNodeToValidate)
 {
-	check(InNode);
+	check(InNodeToValidate);
 
-	TArray<USceneStateMachineNode*> Nodes;
+	// Only consider nodes that share the same type (i.e. name conflicts should only happen between nodes of the same type)
+	auto IsNodeRelevant = [NodeType = InNodeToValidate->GetNodeType()](const USceneStateMachineNode* InStateMachineNode)
+		{
+			return InStateMachineNode && InStateMachineNode->GetNodeType() == NodeType;
+		};
+
+	const USceneStateMachineGraph* Graph = CastChecked<USceneStateMachineGraph>(InNodeToValidate->GetOuter());
 
-	USceneStateMachineGraph* Graph = CastChecked<USceneStateMachineGraph>(InNode->GetOuter());
-	Graph->GetNodesOfClass<USceneStateMachineNode>(Nodes);
+	Names.Reserve(Algo::CountIf(Graph->Nodes,
+		[&IsNodeRelevant](const UEdGraphNode* InNode)
+		{
+			return IsNodeRelevant(Cast<USceneStateMachineNode>(InNode));
+		}));
 
-	Names.Reserve(Nodes.Num());
-	for (USceneStateMachineNode* Node : Nodes)
+	for (const UEdGraphNode* Node : Graph->Nodes)
 	{
-		if (Node != InNode)
+		const USceneStateMachineNode* StateMachineNode = Cast<USceneStateMachineNode>(Node);
+		if (IsNodeRelevant(StateMachineNode))
 		{
-			Names.Add(Node->GetNodeName());
+			Names.Add(StateMachineNode->GetNodeName());
 		}
 	}
 }
diff --git a/Engine/Plugins/Experimental/SceneState/Source/SceneStateMachineGraph/Private/Nodes/SceneStateMachineNodeNameValidator.h b/Engine/Plugins/Experimental/SceneState/Source/SceneStateMachineGraph/Private/Nodes/SceneStateMachineNodeNameValidator.h
--- a/Engine/Plugins/Experimental/SceneState/Source/SceneStateMachineGraph/Private/Nodes/SceneStateMachineNodeNameValidator.h
+++ b/Engine/Plugins/Experimental/SceneState/Source/SceneStateMachineGraph/Private/Nodes/SceneStateMachineNodeNameValidator.h
@@ -14,7 +14,7 @@ namespace UE::SceneState::Graph
 class FStateMachineNodeNameValidator : public INameValidatorInterface
 {
 public:
-	explicit FStateMachineNodeNameValidator(const USceneStateMachineNode* InNode);
+	explicit FStateMachineNodeNameValidator(const USceneStateMachineNode* InNodeToValidate);
 
 private:
 	//~ Begin INameValidatorInterface
diff --git a/Engine/Plugins/Experimental/SceneState/Source/SceneStateMachineGraph/Private/Nodes/SceneStateMachineTaskNode.cpp b/Engine/Plugins/Experimental/SceneState/Source/SceneStateMachineGraph/Private/Nodes/SceneStateMachineTaskNode.cpp
--- a/Engine/Plugins/Experimental/SceneState/Source/SceneStateMachineGraph/Private/Nodes/SceneStateMachineTaskNode.cpp
+++ b/Engine/Plugins/Experimental/SceneState/Source/SceneStateMachineGraph/Private/Nodes/SceneStateMachineTaskNode.cpp
@@ -15,10 +15,9 @@
 
 USceneStateMachineTaskNode::USceneStateMachineTaskNode()
 {
-	NodeName = TEXT("Task");
+	NodeName = NAME_None;
 	NodeType = UE::SceneState::Graph::EStateMachineNodeType::Task;
-
-	bCanRenameNode = false;
+	bCanRenameNode = true;
 
 	UE::SceneState::OnStructIdChanged.AddUObject(this, &USceneStateMachineTaskNode::OnStructIdChanged);
 }
@@ -92,31 +91,43 @@ void USceneStateMachineTaskNode::AllocateDefaultPins()
 	HidePins(MakeArrayView(&USceneStateMachineGraphSchema::PN_In, 1));
 }
 
+void USceneStateMachineTaskNode::OnRenameNode(const FString& InNodeName)
+{
+	Super::OnRenameNode(InNodeName);
+	CachedTaskDisplayName = FText::GetEmpty();
+	OnPostEditTaskDelegate.Broadcast();
+}
+
 FText USceneStateMachineTaskNode::GetNodeTitle(ENodeTitleType::Type InTitleType) const
 {
-	if (CachedNodeTitle.IsEmpty())
+	if (CachedTaskDisplayName.IsEmpty())
 	{
 		if (Task.IsValid())
 		{
 			const FSceneStateTaskDesc& TaskDesc = FSceneStateTaskDescRegistry::Get().GetTaskDesc(Task.GetScriptStruct());
-
-			const FSceneStateTaskDescContext Context = MakeContext();
-
-			FText Tooltip = TaskDesc.GetTooltip(Context);
-			if (!Tooltip.IsEmpty())
-			{
-				Tooltip = FText::Format(INVTEXT("\n{0}"), Tooltip);
-			}
-
-			CachedNodeTitle = FText::Format(LOCTEXT("NodeTitle", "{0}{1}"), TaskDesc.GetDisplayName(Context), Tooltip);
+			CachedTaskDisplayName = TaskDesc.GetDisplayName(MakeContext());
 		}
 		else
 		{
-			CachedNodeTitle = LOCTEXT("NodeTitle_InvalidTaskClass", "Invalid Task Class");
+			CachedTaskDisplayName = LOCTEXT("NodeTitle_InvalidTaskClass", "Invalid Task Class");
 		}
 	}
 
-	return CachedNodeTitle;
+	if (NodeName.IsNone())
+	{
+		return CachedTaskDisplayName;
+	}
+
+	const FText NodeNameText = FText::FromName(NodeName);
+
+	if (InTitleType == ENodeTitleType::FullTitle && !NodeNameText.EqualTo(CachedTaskDisplayName))
+	{
+		return FText::FormatNamed(LOCTEXT("NodeTitle_Full", "{NodeName}\n{TaskDisplayName}")
+			, TEXT("NodeName"), NodeNameText
+			, TEXT("TaskDisplayName"), CachedTaskDisplayName);
+	}
+
+	return NodeNameText;
 }
 
 FText USceneStateMachineTaskNode::GetTooltipText() const
@@ -142,6 +153,11 @@ UObject* USceneStateMachineTaskNode::GetJumpTargetForDoubleClick() const
 	return TaskDesc.GetJumpTarget(MakeContext());
 }
 
+TSharedPtr<INameValidatorInterface> USceneStateMachineTaskNode::MakeNameValidator() const
+{
+	return MakeShared<FDummyNameValidator>(EValidatorResult::Ok);
+}
+
 void USceneStateMachineTaskNode::PostLoad()
 {
 	Super::PostLoad();
@@ -160,7 +176,7 @@ void USceneStateMachineTaskNode::PostEditChangeProperty(FPropertyChangedEvent& I
 {
 	Super::PostEditChangeProperty(InPropertyChangedEvent);
 
-	CachedNodeTitle = FText::GetEmpty();
+	CachedTaskDisplayName = FText::GetEmpty();
 
 	if (InPropertyChangedEvent.GetMemberPropertyName() == GET_MEMBER_NAME_CHECKED(USceneStateMachineTaskNode, Task))
 	{
diff --git a/Engine/Plugins/Experimental/SceneState/Source/SceneStateMachineGraph/Public/Nodes/SceneStateMachineTaskNode.h b/Engine/Plugins/Experimental/SceneState/Source/SceneStateMachineGraph/Public/Nodes/SceneStateMachineTaskNode.h
--- a/Engine/Plugins/Experimental/SceneState/Source/SceneStateMachineGraph/Public/Nodes/SceneStateMachineTaskNode.h
+++ b/Engine/Plugins/Experimental/SceneState/Source/SceneStateMachineGraph/Public/Nodes/SceneStateMachineTaskNode.h
@@ -53,11 +53,13 @@ public:
 
 	//~ Begin UEdGraphNode
 	virtual void AllocateDefaultPins() override;
+	virtual void OnRenameNode(const FString& InNodeName);
 	virtual FText GetNodeTitle(ENodeTitleType::Type InTitleType) const override;
 	virtual FText GetTooltipText() const override;
 	virtual bool CanDuplicateNode() const override;
 	virtual void PostPasteNode() override;
 	virtual UObject* GetJumpTargetForDoubleClick() const override;
+	virtual TSharedPtr<INameValidatorInterface> MakeNameValidator() const override;
 	//~ End UEdGraphNode
 
 	//~ Begin UObject
@@ -108,5 +110,5 @@ private:
 
 	mutable TMulticastDelegate<void()> OnPostEditTaskDelegate;
 
-	mutable FText CachedNodeTitle;
+	mutable FText CachedTaskDisplayName;
 };
