Motion Design Scene State: fixed issue where removing a graph would be kept in transaction and new graphs wouldn't be able to be named after that removed graph until the removed graph is fully gc'd (after transaction clear)

diff --git a/Engine/Plugins/Experimental/SceneState/Source/SceneStateBlueprintEditor/Private/Menus/SceneStateMachineContextMenu.cpp b/Engine/Plugins/Experimental/SceneState/Source/SceneStateBlueprintEditor/Private/Menus/SceneStateMachineContextMenu.cpp
--- a/Engine/Plugins/Experimental/SceneState/Source/SceneStateBlueprintEditor/Private/Menus/SceneStateMachineContextMenu.cpp
+++ b/Engine/Plugins/Experimental/SceneState/Source/SceneStateBlueprintEditor/Private/Menus/SceneStateMachineContextMenu.cpp
@@ -9,6 +9,7 @@
 #include "SceneStateBlueprint.h"
 #include "SceneStateBlueprintEditor.h"
 #include "SceneStateMachineGraphUtils.h"
+#include "SceneStateUtils.h"
 #include "ScopedTransaction.h"
 #include "ToolMenu.h"
 #include "ToolMenuSection.h"
@@ -168,6 +169,8 @@ void FStateMachineContextMenu::Delete()
 
 	BlueprintEditor->CloseDocumentTab(GraphAction->EdGraph);
 	FBlueprintEditorUtils::MarkBlueprintAsStructurallyModified(Blueprint);
+
+	DiscardObject(GraphAction->EdGraph);
 }
 
 } // UE::SceneState::Editor
