Motion Design Scene State: added ability to reorder state machines in list

diff --git a/Engine/Plugins/Experimental/SceneState/Source/SceneStateBlueprint/Private/Actions/SceneStateBlueprintAction_Graph.cpp b/Engine/Plugins/Experimental/SceneState/Source/SceneStateBlueprint/Private/Actions/SceneStateBlueprintAction_Graph.cpp
new file mode 100644
--- /dev/null
+++ b/Engine/Plugins/Experimental/SceneState/Source/SceneStateBlueprint/Private/Actions/SceneStateBlueprintAction_Graph.cpp
@@ -0,0 +1,100 @@
+// Copyright Epic Games, Inc. All Rights Reserved.
+
+#include "Actions/SceneStateBlueprintAction_Graph.h"
+#include "Kismet2/BlueprintEditorUtils.h"
+#include "SceneStateBlueprintUtils.h"
+#include "SceneStateMachineGraph.h"
+#include "ScopedTransaction.h"
+
+#define LOCTEXT_NAMESPACE "SceneStateBlueprintAction_Graph"
+
+namespace UE::SceneState::Graph
+{
+
+FBlueprintAction_Graph::FBlueprintAction_Graph(const FArguments& InArgs)
+{
+	EdGraph = InArgs.Graph;
+	GraphType = InArgs.GraphType;
+	Grouping = InArgs.Grouping;
+	SectionID = InArgs.SectionID;
+	Icon = InArgs.Icon;
+
+	if (EdGraph)
+	{
+		FuncName = EdGraph->GetFName();
+
+		if (const UEdGraphSchema* Schema = EdGraph->GetSchema())
+		{
+			FGraphDisplayInfo DisplayInfo;
+			Schema->GetGraphDisplayInformation(*EdGraph, DisplayInfo);
+
+			UpdateSearchData(MoveTemp(DisplayInfo.DisplayName), MoveTemp(DisplayInfo.Tooltip), InArgs.Category, FText());
+		}
+	}
+}
+
+FName FBlueprintAction_Graph::StaticGetTypeId()
+{
+	static FName Type = TEXT("FEdGraphSchemaAction");
+	return Type;
+}
+
+const FSlateBrush* FBlueprintAction_Graph::GetPaletteIcon() const
+{
+	return Icon.GetIcon();
+}
+
+FName FBlueprintAction_Graph::GetTypeId() const
+{
+	return StaticGetTypeId();
+}
+
+void FBlueprintAction_Graph::MovePersistentItemToCategory(const FText& InNewCategory)
+{
+	USceneStateMachineGraph* const StateMachineGraph = Cast<USceneStateMachineGraph>(EdGraph);
+
+	// Only allow top-level state machines to have a category
+	if (StateMachineGraph && GraphType == EEdGraphSchemaAction_K2Graph::Graph)
+	{
+		UBlueprint* Blueprint = FBlueprintEditorUtils::FindBlueprintForGraphChecked(EdGraph);
+
+		FScopedTransaction Transaction(LOCTEXT("SetStateMachineGraphCategory", "Set Category"));
+		StateMachineGraph->Modify();
+		StateMachineGraph->Category = InNewCategory;
+		FBlueprintEditorUtils::MarkBlueprintAsModified(Blueprint);
+	}
+}
+
+int32 FBlueprintAction_Graph::GetReorderIndexInContainer() const
+{
+	return FindIndexOfGraphInParent(EdGraph);
+}
+
+bool FBlueprintAction_Graph::ReorderToBeforeAction(TSharedRef<FEdGraphSchemaAction> OtherAction)
+{
+	if (OtherAction->GetTypeId() != GetTypeId())
+	{
+		return false;
+	}
+
+	const int32 CurrentIndex = GetReorderIndexInContainer();
+	const int32 TargetIndex = OtherAction->GetReorderIndexInContainer();
+
+	if (CurrentIndex == INDEX_NONE || TargetIndex == INDEX_NONE || CurrentIndex == TargetIndex)
+	{
+		return false;
+	}
+
+	if (MoveGraph(EdGraph, TargetIndex))
+	{
+		// Change category to match the one we dragged on to as well
+		MovePersistentItemToCategory(OtherAction->GetCategory());
+		return true;
+	}
+
+	return false;
+}
+
+} // UE::SceneState::Graph
+
+#undef LOCTEXT_NAMESPACE
diff --git a/Engine/Plugins/Experimental/SceneState/Source/SceneStateBlueprint/Private/SceneStateBlueprintBindingUtils.cpp b/Engine/Plugins/Experimental/SceneState/Source/SceneStateBlueprint/Private/SceneStateBlueprintBindingUtils.cpp
--- a/Engine/Plugins/Experimental/SceneState/Source/SceneStateBlueprint/Private/SceneStateBlueprintBindingUtils.cpp
+++ b/Engine/Plugins/Experimental/SceneState/Source/SceneStateBlueprint/Private/SceneStateBlueprintBindingUtils.cpp
@@ -56,7 +56,7 @@ void GetStateBindingDescs(const USceneStateMachineStateNode* InStateNode, TArray
 	OutBindingDescs.Emplace(TInstancedStruct<FSceneStateBindingDesc>::Make(CreateBindingDesc(*StateMachineGraph)));
 
 	// Recurse up to add the Event Handlers of the parent state node (if any)
-	GetStateBindingDescs(Cast<USceneStateMachineStateNode>(StateMachineGraph->GetOuter()), OutBindingDescs, CategoryName);
+	GetStateBindingDescs(StateMachineGraph->GetParentStateNode(), OutBindingDescs, CategoryName);
 }
 
 FSceneStateBindingDesc CreateBindingDesc(const USceneStateBlueprint& InBlueprint)
diff --git a/Engine/Plugins/Experimental/SceneState/Source/SceneStateBlueprint/Private/SceneStateBlueprintUtils.cpp b/Engine/Plugins/Experimental/SceneState/Source/SceneStateBlueprint/Private/SceneStateBlueprintUtils.cpp
--- a/Engine/Plugins/Experimental/SceneState/Source/SceneStateBlueprint/Private/SceneStateBlueprintUtils.cpp
+++ b/Engine/Plugins/Experimental/SceneState/Source/SceneStateBlueprint/Private/SceneStateBlueprintUtils.cpp
@@ -130,4 +130,41 @@ void CreateBlueprintVariables(USceneStateBlueprint* InBlueprint, TArrayView<UE::
 	}
 }
 
+int32 FindIndexOfGraphInParent(UEdGraph* InGraph)
+{
+	if (USceneStateBlueprint* Blueprint = Cast<USceneStateBlueprint>(FBlueprintEditorUtils::FindBlueprintForGraph(InGraph)))
+	{
+		// Only consider top-level state machines as reorderable.
+		const int32 Index = Blueprint->StateMachineGraphs.IndexOfByKey(InGraph);
+		if (Index != INDEX_NONE)
+		{
+			return Index;
+		}
+	}
+	return FBlueprintEditorUtils::FindIndexOfGraphInParent(InGraph);
+}
+
+bool MoveGraph(UEdGraph* InGraph, int32 InTargetIndex)
+{
+	USceneStateBlueprint* Blueprint = Cast<USceneStateBlueprint>(FBlueprintEditorUtils::FindBlueprintForGraph(InGraph));
+	if (!Blueprint)
+	{
+		return false;
+	}
+
+	// Deal with State Machine Graphs
+	const int32 CurrentIndex = Blueprint->StateMachineGraphs.IndexOfByKey(InGraph);
+	if (CurrentIndex != INDEX_NONE && CurrentIndex != InTargetIndex && Blueprint->StateMachineGraphs.IsValidIndex(InTargetIndex))
+	{
+		Blueprint->Modify();
+		Blueprint->StateMachineGraphs.Insert(InGraph, InTargetIndex);
+		Blueprint->StateMachineGraphs.RemoveAt(CurrentIndex < InTargetIndex ? CurrentIndex : CurrentIndex + 1);
+		FBlueprintEditorUtils::MarkBlueprintAsModified(Blueprint);
+		return true;
+	}
+
+	// Deal with Other Graphs in Blueprint
+	return FBlueprintEditorUtils::MoveGraphBeforeOtherGraph(InGraph, InTargetIndex, /*bDontRecompile*/true);
+}
+
 } // namespace UE::SceneState::Graph
diff --git a/Engine/Plugins/Experimental/SceneState/Source/SceneStateBlueprint/Private/SceneStateBlueprintUtils.h b/Engine/Plugins/Experimental/SceneState/Source/SceneStateBlueprint/Private/SceneStateBlueprintUtils.h
--- a/Engine/Plugins/Experimental/SceneState/Source/SceneStateBlueprint/Private/SceneStateBlueprintUtils.h
+++ b/Engine/Plugins/Experimental/SceneState/Source/SceneStateBlueprint/Private/SceneStateBlueprintUtils.h
@@ -46,4 +46,10 @@ namespace UE::SceneState::Graph
 	 */
 	void CreateBlueprintVariables(USceneStateBlueprint* InBlueprint, TArrayView<UE::PropertyBinding::FPropertyCreationDescriptor> InPropertyCreationDescs);
 
+	/** Finds the index of the specified graph in the parent (if it is not reorderable, then we will return INDEX_NONE) */
+	int32 FindIndexOfGraphInParent(UEdGraph* InGraph);
+
+	/** Reorders the specified graph to be at the given target index in the parent */
+	bool MoveGraph(UEdGraph* InGraph, int32 InTargetIndex);
+
 } // UE::SceneState::Graph
diff --git a/Engine/Plugins/Experimental/SceneState/Source/SceneStateBlueprint/Public/Actions/SceneStateBlueprintAction_Graph.h b/Engine/Plugins/Experimental/SceneState/Source/SceneStateBlueprint/Public/Actions/SceneStateBlueprintAction_Graph.h
--- a/Engine/Plugins/Experimental/SceneState/Source/SceneStateBlueprint/Public/Actions/SceneStateBlueprintAction_Graph.h
+++ b/Engine/Plugins/Experimental/SceneState/Source/SceneStateBlueprint/Public/Actions/SceneStateBlueprintAction_Graph.h
@@ -5,6 +5,8 @@
 #include "EdGraphSchema_K2_Actions.h"
 #include "Textures/SlateIcon.h"
 
+#define UE_API SCENESTATEBLUEPRINT_API
+
 namespace UE::SceneState::Graph
 {
 
@@ -12,37 +14,31 @@ struct FBlueprintAction_Graph : public FEdGraphSchemaAction_K2Graph
 {
 	FBlueprintAction_Graph() = default;
 
-	explicit FBlueprintAction_Graph(EEdGraphSchemaAction_K2Graph::Type InType
-		, const FText& InNodeCategory
-		, const FText& InMenuDesc
-		, const FText& InToolTip
-		, const FSlateIcon& InIcon
-		, const int32 InGrouping
-		, const int32 InSectionID = 0)
-		: FEdGraphSchemaAction_K2Graph(InType, InNodeCategory, InMenuDesc, InToolTip, InGrouping, InSectionID)
-		, Icon(InIcon)
+	struct FArguments
 	{
-	}
+		UEdGraph* Graph = nullptr;
+		EEdGraphSchemaAction_K2Graph::Type GraphType = EEdGraphSchemaAction_K2Graph::Graph;
+		FText Category;
+		FSlateIcon Icon;
+		int32 Grouping = 0;
+		int32 SectionID = 0;
+	};
+	UE_API explicit FBlueprintAction_Graph(const FArguments& InArgs);
 
-	static FName StaticGetTypeId()
-	{
-		static FName Type = TEXT("FEdGraphSchemaAction");
-		return Type;
-	}
+	UE_API static FName StaticGetTypeId();
 
 	//~ Begin FEdGraphSchemaAction
-	virtual const FSlateBrush* GetPaletteIcon() const override
-	{
-		return Icon.GetIcon();
-	}
-
-	virtual FName GetTypeId() const override
-	{
-		return StaticGetTypeId();
-	}
+	UE_API virtual const FSlateBrush* GetPaletteIcon() const override;
+	UE_API virtual FName GetTypeId() const override;
+	UE_API virtual void MovePersistentItemToCategory(const FText& InNewCategory) override;
+	UE_API virtual int32 GetReorderIndexInContainer() const override;
+	UE_API virtual bool ReorderToBeforeAction(TSharedRef<FEdGraphSchemaAction> OtherAction) override;
 	//~ End FEdGraphSchemaAction
 
-	const FSlateIcon Icon;
+private:
+	FSlateIcon Icon;
 };
 
 } // UE::SceneState::Editor
+
+#undef UE_API
diff --git a/Engine/Plugins/Experimental/SceneState/Source/SceneStateBlueprintEditor/Private/DetailsView/SceneStateMachineGraphCustomization.cpp b/Engine/Plugins/Experimental/SceneState/Source/SceneStateBlueprintEditor/Private/DetailsView/SceneStateMachineGraphCustomization.cpp
--- a/Engine/Plugins/Experimental/SceneState/Source/SceneStateBlueprintEditor/Private/DetailsView/SceneStateMachineGraphCustomization.cpp
+++ b/Engine/Plugins/Experimental/SceneState/Source/SceneStateBlueprintEditor/Private/DetailsView/SceneStateMachineGraphCustomization.cpp
@@ -16,32 +16,47 @@ namespace UE::SceneState::Editor
 
 void FStateMachineGraphCustomization::CustomizeDetails(IDetailLayoutBuilder& InDetailBuilder)
 {
-	TSharedRef<IPropertyHandle> StateMachineIdHandle = InDetailBuilder.GetProperty(GET_MEMBER_NAME_CHECKED(USceneStateMachineGraph, ParametersId));
-	TSharedRef<IPropertyHandle> ParametersHandle = InDetailBuilder.GetProperty(GET_MEMBER_NAME_CHECKED(USceneStateMachineGraph, Parameters));
-	TSharedRef<IPropertyHandle> RunModeHandle = InDetailBuilder.GetProperty(GET_MEMBER_NAME_CHECKED(USceneStateMachineGraph, RunMode));
+	const TArray<TWeakObjectPtr<USceneStateMachineGraph>> Graphs = InDetailBuilder.GetObjectsOfTypeBeingCustomized<USceneStateMachineGraph>();
 
+	const TSharedRef<IPropertyHandle> CategoryHandle = InDetailBuilder.GetProperty(GET_MEMBER_NAME_CHECKED(USceneStateMachineGraph, Category));
+	const TSharedRef<IPropertyHandle> StateMachineIdHandle = InDetailBuilder.GetProperty(GET_MEMBER_NAME_CHECKED(USceneStateMachineGraph, ParametersId));
+	const TSharedRef<IPropertyHandle> ParametersHandle = InDetailBuilder.GetProperty(GET_MEMBER_NAME_CHECKED(USceneStateMachineGraph, Parameters));
+	const TSharedRef<IPropertyHandle> RunModeHandle = InDetailBuilder.GetProperty(GET_MEMBER_NAME_CHECKED(USceneStateMachineGraph, RunMode));
+
+	CategoryHandle->MarkHiddenByCustomization();
 	StateMachineIdHandle->MarkHiddenByCustomization();
 	ParametersHandle->MarkHiddenByCustomization();
 	RunModeHandle->MarkHiddenByCustomization();
 
 	// Listen to Parameter Changes to broadcast change
-	{
-		TArray<TWeakObjectPtr<USceneStateMachineGraph>> Graphs = InDetailBuilder.GetObjectsOfTypeBeingCustomized<USceneStateMachineGraph>();
-		ParametersHandle->SetOnPropertyValueChanged(FSimpleDelegate::CreateSPLambda(this,
-			[Graphs = MoveTemp(Graphs)]()
+	ParametersHandle->SetOnPropertyValueChanged(FSimpleDelegate::CreateSPLambda(this,
+		[Graphs]()
+		{
+			for (TWeakObjectPtr<USceneStateMachineGraph> GraphWeak : Graphs)
 			{
-				for (TWeakObjectPtr<USceneStateMachineGraph> GraphWeak : Graphs)
+				if (USceneStateMachineGraph* Graph = GraphWeak.Get())
 				{
-					if (USceneStateMachineGraph* Graph = GraphWeak.Get())
-					{
-						Graph->NotifyParametersChanged();
-					}
+					Graph->NotifyParametersChanged();
 				}
-			}));
-	}
+			}
+		}));
 
 	// State Machine Category
 	IDetailCategoryBuilder& StateMachineCategory = InDetailBuilder.EditCategory(TEXT("State Machine"));
+
+	// Show category if there's at least one top-level graph
+	const bool bShowCategory = Graphs.ContainsByPredicate(
+		[](TWeakObjectPtr<USceneStateMachineGraph> InGraphWeak)
+		{
+			USceneStateMachineGraph* Graph = InGraphWeak.Get();
+			return Graph && Graph->IsRootStateMachine();
+		});
+
+	if (bShowCategory)
+	{
+		StateMachineCategory.AddProperty(CategoryHandle);
+	}
+
 	StateMachineCategory.AddProperty(RunModeHandle);
 
 	FGuid StateMachineId;
diff --git a/Engine/Plugins/Experimental/SceneState/Source/SceneStateBlueprintEditor/Private/Widgets/SSceneStateStateMachineMenu.cpp b/Engine/Plugins/Experimental/SceneState/Source/SceneStateBlueprintEditor/Private/Widgets/SSceneStateStateMachineMenu.cpp
--- a/Engine/Plugins/Experimental/SceneState/Source/SceneStateBlueprintEditor/Private/Widgets/SSceneStateStateMachineMenu.cpp
+++ b/Engine/Plugins/Experimental/SceneState/Source/SceneStateBlueprintEditor/Private/Widgets/SSceneStateStateMachineMenu.cpp
@@ -15,7 +15,9 @@
 #include "SceneStateBlueprintEditorStyle.h"
 #include "SceneStateMachineGraph.h"
 #include "SceneStateMachineGraphSchema.h"
+#include "SceneStateMachineItemDragDropOp.h"
 #include "SceneStateTransitionGraphSchema.h"
+#include "ScopedTransaction.h"
 #include "Widgets/Input/SButton.h"
 #include "Widgets/Input/SSearchBox.h"
 #include "Widgets/Layout/SBorder.h"
@@ -57,11 +59,13 @@ void SStateMachineMenu::Construct(const FArguments& InArgs, const TSharedRef<FSc
 		.OnCollectAllActions(this, &SStateMachineMenu::CollectGraphActions)
 		.OnCollectStaticSections(this, &SStateMachineMenu::CollectSections)
 		.OnActionSelected(this, &SStateMachineMenu::OnGraphActionSelected)
+		.OnActionDragged(this, &SStateMachineMenu::OnGraphActionDragged)
 		.OnActionDoubleClicked(this, &SStateMachineMenu::OnGraphActionDoubleClicked)
 		.OnContextMenuOpening(this, &SStateMachineMenu::OnContextMenuOpening)
 		.OnGetSectionTitle(this, &SStateMachineMenu::GetSectionTitle)
 		.OnGetSectionWidget(this, &SStateMachineMenu::CreateSectionWidget)
 		.OnCanRenameSelectedAction(this, &SStateMachineMenu::CanRequestRenameOnActionNode)
+		.OnCategoryTextCommitted(this, &SStateMachineMenu::OnCategoryTextCommitted)
 		.DefaultRowExpanderBaseIndentLevel(1)
 		.AlphaSortItems(false)
 		.UseSectionStyling(true);
@@ -247,6 +251,23 @@ void SStateMachineMenu::OnGraphActionSelected(const TArray<TSharedPtr<FEdGraphSc
 	BlueprintEditor->GetInspector()->ShowDetailsForSingleObject(DetailsObject, SKismetInspector::FShowDetailsOptions(DetailsText));
 }
 
+FReply SStateMachineMenu::OnGraphActionDragged(const TArray<TSharedPtr<FEdGraphSchemaAction>>& InActions, const FPointerEvent& InMouseEvent)
+{
+	if (InActions.IsEmpty() || !InActions[0].IsValid())
+	{
+		return FReply::Unhandled();
+	}
+
+	const TSharedPtr<FEdGraphSchemaAction>& Action = InActions[0];
+
+	if (Action->GetTypeId() == Graph::FBlueprintAction_Graph::StaticGetTypeId())
+	{
+		return FReply::Handled().BeginDragDrop(FStateMachineItemDragDropOp::New(Action));
+	}
+
+	return FReply::Unhandled();
+}
+
 bool SStateMachineMenu::ShouldProcessGraph(UEdGraph* InGraph) const
 {
 	if (!InGraph)
@@ -301,30 +322,25 @@ void SStateMachineMenu::CollectGraphActionsRecursive(UEdGraph* InGraph, FText In
 		return;
 	}
 
-	FGraphDisplayInfo DisplayInfo;
-	Schema->GetGraphDisplayInformation(*InGraph, DisplayInfo);
-
-	EEdGraphSchemaAction_K2Graph::Type GraphType = static_cast<EEdGraphSchemaAction_K2Graph::Type>(InGraphType);
-
-	TSharedPtr<Graph::FBlueprintAction_Graph> GraphAction = MakeShared<Graph::FBlueprintAction_Graph>(GraphType
-		, InCategory
-		, DisplayInfo.DisplayName
-		, DisplayInfo.Tooltip
-		, FBlueprintEditorStyle::Get().GetGraphSchemaIcon(Schema->GetClass())
-		, 1
-		, NodeSectionID::GRAPH);
+	const Graph::FBlueprintAction_Graph::FArguments Arguments {
+			.Graph = InGraph,
+			.GraphType = static_cast<EEdGraphSchemaAction_K2Graph::Type>(InGraphType),
+			.Category = InCategory,
+			.Icon = FBlueprintEditorStyle::Get().GetGraphSchemaIcon(Schema->GetClass()),
+			.Grouping = 1,
+			.SectionID = NodeSectionID::GRAPH,
+		};
 
-	GraphAction->FuncName = InGraph->GetFName();
-	GraphAction->EdGraph = InGraph;
+	TSharedRef<Graph::FBlueprintAction_Graph> GraphAction = MakeShared<Graph::FBlueprintAction_Graph>(Arguments);
 	OutActions.AddAction(GraphAction);
 
 	if (InCategory.IsEmpty())
 	{
-		InCategory = DisplayInfo.DisplayName;
+		InCategory = GraphAction->GetCategory();
 	}
 	else
 	{
-		InCategory = FText::Format(INVTEXT("{0}|{1}"), InCategory, DisplayInfo.DisplayName);
+		InCategory = FText::Format(INVTEXT("{0}|{1}"), InCategory, GraphAction->GetMenuDescription());
 	}
 
 	for (UEdGraph* SubGraph : InGraph->SubGraphs)
@@ -339,7 +355,10 @@ void SStateMachineMenu::CollectGraphActions(FGraphActionListBuilderBase& OutActi
 	{
 		for (UEdGraph* Graph : Blueprint->StateMachineGraphs)
 		{
-			CollectGraphActionsRecursive(Graph, FText::GetEmpty(), EEdGraphSchemaAction_K2Graph::Type::Graph, OutActions);
+			if (USceneStateMachineGraph* StateMachineGraph = Cast<USceneStateMachineGraph>(Graph))
+			{
+				CollectGraphActionsRecursive(StateMachineGraph, StateMachineGraph->Category, EEdGraphSchemaAction_K2Graph::Type::Graph, OutActions);
+			}
 		}
 	}
 }
@@ -423,6 +442,18 @@ bool SStateMachineMenu::CanRequestRenameOnActionNode(TWeakPtr<FGraphActionNode>
 		return false;
 	}
 
+	if (GraphActionNode->IsCategoryNode())
+	{
+		TArray<TSharedPtr<FEdGraphSchemaAction>> Actions;
+		GraphActionMenu->GetCategorySubActions(InSelectedNodeWeak, Actions);
+
+		return Actions.ContainsByPredicate(
+			[](const TSharedPtr<FEdGraphSchemaAction>& InAction)
+			{
+				return InAction.IsValid() && InAction->GetPersistentItemDefiningObject().IsPotentiallyEditable();
+			});
+	}
+
 	if (GraphActionNode->IsActionNode())
 	{
 		if (!FBlueprintEditorUtils::IsPaletteActionReadOnly(GraphActionNode->Action, BlueprintEditor))
@@ -434,6 +465,47 @@ bool SStateMachineMenu::CanRequestRenameOnActionNode(TWeakPtr<FGraphActionNode>
 	return false;
 }
 
+void SStateMachineMenu::OnCategoryTextCommitted(const FText& InText, ETextCommit::Type InCommitType, TWeakPtr<FGraphActionNode> InTargetAction)
+{
+	const TSharedPtr<FGraphActionNode> TargetAction = InTargetAction.Pin();
+	if (!TargetAction.IsValid())
+	{
+		return;
+	}
+
+	UBlueprint* const Blueprint = GetBlueprint();
+	if (!Blueprint)
+	{
+		return;
+	}
+
+	// Remove excess whitespace and prevent categories with just spaces
+	const FText CategoryName = FText::TrimPrecedingAndTrailing(InText);
+
+	TArray<TSharedPtr<FEdGraphSchemaAction>> Actions;
+	GraphActionMenu->GetCategorySubActions(InTargetAction, Actions);
+
+	if (Actions.IsEmpty())
+	{
+		return;
+	}
+
+	FScopedTransaction Transaction(LOCTEXT("RenameCategory", "Rename Category"));
+	Blueprint->Modify();
+
+	for (const TSharedPtr<FEdGraphSchemaAction>& Action : Actions)
+	{
+		if (Action)
+		{
+			Action->MovePersistentItemToCategory(CategoryName);
+		}
+	}
+
+	RefreshMenu();
+	FBlueprintEditorUtils::MarkBlueprintAsModified(Blueprint);
+	GraphActionMenu->SelectItemByName(*CategoryName.ToString(), ESelectInfo::OnMouseClick, TargetAction->SectionID, /*bIsCategory*/true);
+}
+
 void SStateMachineMenu::Tick(const FGeometry& InGeometry, const double InCurrentTime, const float InDeltaTime)
 {
 	if (bPendingRefresh)
diff --git a/Engine/Plugins/Experimental/SceneState/Source/SceneStateBlueprintEditor/Private/Widgets/SSceneStateStateMachineMenu.h b/Engine/Plugins/Experimental/SceneState/Source/SceneStateBlueprintEditor/Private/Widgets/SSceneStateStateMachineMenu.h
--- a/Engine/Plugins/Experimental/SceneState/Source/SceneStateBlueprintEditor/Private/Widgets/SSceneStateStateMachineMenu.h
+++ b/Engine/Plugins/Experimental/SceneState/Source/SceneStateBlueprintEditor/Private/Widgets/SSceneStateStateMachineMenu.h
@@ -62,6 +62,9 @@ private:
 
 	void OnGraphActionSelected(const TArray<TSharedPtr<FEdGraphSchemaAction>>& InActions, ESelectInfo::Type InSelectionType);
 
+	/** Called when dragging the graph action entry */
+	FReply OnGraphActionDragged(const TArray<TSharedPtr<FEdGraphSchemaAction>>& InActions, const FPointerEvent& InMouseEvent);
+
 	bool ShouldProcessGraph(UEdGraph* InGraph) const;
 
 	void CollectGraphActionsRecursive(UEdGraph* InGraph, FText InCategory, int32 InGraphType, FGraphActionListBuilderBase& OutActions);
@@ -82,6 +85,9 @@ private:
 
 	bool CanRequestRenameOnActionNode(TWeakPtr<FGraphActionNode> InSelectedNodeWeak) const;
 
+	/** Called when a category target item has been renamed */
+	void OnCategoryTextCommitted(const FText& InText, ETextCommit::Type InCommitType, TWeakPtr<FGraphActionNode> InTargetAction);
+
 	//~ Begin SWidget
 	virtual void Tick(const FGeometry& InGeometry, const double InCurrentTime, const float InDeltaTime) override;
 	virtual FReply OnKeyDown(const FGeometry& InGeometry, const FKeyEvent& InKeyEvent) override;
diff --git a/Engine/Plugins/Experimental/SceneState/Source/SceneStateBlueprintEditor/Private/Widgets/SceneStateMachineItemDragDropOp.cpp b/Engine/Plugins/Experimental/SceneState/Source/SceneStateBlueprintEditor/Private/Widgets/SceneStateMachineItemDragDropOp.cpp
new file mode 100644
--- /dev/null
+++ b/Engine/Plugins/Experimental/SceneState/Source/SceneStateBlueprintEditor/Private/Widgets/SceneStateMachineItemDragDropOp.cpp
@@ -0,0 +1,136 @@
+// Copyright Epic Games, Inc. All Rights Reserved.
+
+#include "SceneStateMachineItemDragDropOp.h"
+#include "ScopedTransaction.h"
+#include "Styling/AppStyle.h"
+
+#define LOCTEXT_NAMESPACE "SceneStateMachineItemDragDropOp"
+
+namespace UE::SceneState::Editor
+{
+
+TSharedRef<FGraphSchemaActionDragDropAction> FStateMachineItemDragDropOp::New(const TSharedPtr<FEdGraphSchemaAction>& InActionNode)
+{
+	TSharedRef<FStateMachineItemDragDropOp> DragDropOp = MakeShared<FStateMachineItemDragDropOp>();
+	DragDropOp->SourceAction = InActionNode;
+	DragDropOp->Construct();
+	return DragDropOp;
+}
+
+void FStateMachineItemDragDropOp::HoverTargetChanged()
+{
+	if (!SourceAction.IsValid())
+	{
+		return;
+	}
+
+	if (!HoveredCategoryName.IsEmpty())
+	{
+		OnHoverCategory(HoveredCategoryName);
+	}
+	else if (const TSharedPtr<FEdGraphSchemaAction> TargetAction = HoveredAction.Pin())
+	{
+		OnHoverTargetAction(TargetAction.ToSharedRef());
+	}
+}
+
+FReply FStateMachineItemDragDropOp::DroppedOnCategory(FText InCategory)
+{
+	if (!SourceAction.IsValid())
+	{
+		return FReply::Unhandled();
+	}
+
+	SourceAction->MovePersistentItemToCategory(InCategory);
+	return FReply::Handled();
+}
+
+FReply FStateMachineItemDragDropOp::DroppedOnAction(TSharedRef<FEdGraphSchemaAction> InAction)
+{
+	if (!SourceAction.IsValid()
+		|| SourceAction->GetTypeId() != InAction->GetTypeId()
+		|| SourceAction->GetPersistentItemDefiningObject() != InAction->GetPersistentItemDefiningObject())
+	{
+		return FReply::Unhandled();
+	}
+
+	FScopedTransaction Transaction(LOCTEXT("ReorderStateMachineGraph", "Reorder State Machine Graph"));
+	if (!SourceAction->ReorderToBeforeAction(InAction))
+	{
+		Transaction.Cancel();
+	}
+
+	return FReply::Handled();
+}
+
+void FStateMachineItemDragDropOp::OnHoverCategory(const FText& InCategory)
+{
+	FFormatNamedArguments Args;
+	Args.Add(TEXT("SourceDisplayName"), SourceAction->GetMenuDescription());
+	Args.Add(TEXT("TargetCategoryName"), HoveredCategoryName);
+
+	if (HoveredCategoryName.EqualTo(SourceAction->GetCategory()))
+	{
+		SetFeedbackMessageError(FText::Format(LOCTEXT("SameCategoryError", "'{SourceDisplayName}' is already in category '{TargetCategoryName}'"), Args));
+	}
+	else
+	{
+		SetFeedbackMessageOk(FText::Format(LOCTEXT("SetCategoryOk", "Move '{SourceDisplayName}' to category '{TargetCategoryName}'"), Args));
+	}
+}
+
+void FStateMachineItemDragDropOp::OnHoverTargetAction(const TSharedRef<FEdGraphSchemaAction>& InTargetAction)
+{
+	FFormatNamedArguments Args;
+	Args.Add(TEXT("SourceDisplayName"), SourceAction->GetMenuDescription());
+	Args.Add(TEXT("TargetDisplayName"), InTargetAction->GetMenuDescription());
+
+	if (InTargetAction->GetTypeId() != SourceAction->GetTypeId())
+	{
+		SetFeedbackMessageError(FText::Format(LOCTEXT("ReorderActionDifferentAction", "Cannot reorder '{SourceDisplayName}' into a different section."), Args));
+		return;
+	}
+
+	if (SourceAction->GetPersistentItemDefiningObject() != InTargetAction->GetPersistentItemDefiningObject())
+	{
+		SetFeedbackMessageError(FText::Format(LOCTEXT("ReorderActionDifferentScope", "Cannot reorder '{SourceDisplayName}' into a different scope."), Args));
+		return;
+	}
+
+	const int32 SourceIndex = SourceAction->GetReorderIndexInContainer();
+	const int32 TargetIndex = InTargetAction->GetReorderIndexInContainer();
+
+	if (SourceIndex == INDEX_NONE)
+	{
+		SetFeedbackMessageError(FText::Format(LOCTEXT("ReorderNonOrderedItem", "Cannot reorder '{SourceDisplayName}'."), Args));
+		return;
+	}
+
+	if (TargetIndex == INDEX_NONE)
+	{
+		SetFeedbackMessageError(FText::Format(LOCTEXT("ReorderOntoNonOrderedItem", "Cannot reorder '{SourceDisplayName}' before '{TargetDisplayName}'."), Args));
+		return;
+	}
+
+	if (InTargetAction == SourceAction)
+	{
+		SetFeedbackMessageError(FText::Format(LOCTEXT("ReorderOntoSameItem", "Cannot reorder '{SourceDisplayName}' before itself."), Args));
+		return;
+	}
+
+	SetFeedbackMessageOk(FText::Format(LOCTEXT("ReorderActionOk", "Reorder '{SourceDisplayName}' before '{TargetDisplayName}'"), Args));
+}
+
+void FStateMachineItemDragDropOp::SetFeedbackMessageOk(const FText& InMessage)
+{
+	SetSimpleFeedbackMessage(FAppStyle::GetBrush(TEXT("Graph.ConnectorFeedback.OK")), FLinearColor::White, InMessage);
+}
+
+void FStateMachineItemDragDropOp::SetFeedbackMessageError(const FText& InMessage)
+{
+	SetSimpleFeedbackMessage(FAppStyle::GetBrush(TEXT("Graph.ConnectorFeedback.Error")), FLinearColor::White, InMessage);
+}
+
+} // UE::SceneState::Editor
+
+#undef LOCTEXT_NAMESPACE
diff --git a/Engine/Plugins/Experimental/SceneState/Source/SceneStateBlueprintEditor/Private/Widgets/SceneStateMachineItemDragDropOp.h b/Engine/Plugins/Experimental/SceneState/Source/SceneStateBlueprintEditor/Private/Widgets/SceneStateMachineItemDragDropOp.h
new file mode 100644
--- /dev/null
+++ b/Engine/Plugins/Experimental/SceneState/Source/SceneStateBlueprintEditor/Private/Widgets/SceneStateMachineItemDragDropOp.h
@@ -0,0 +1,42 @@
+// Copyright Epic Games, Inc. All Rights Reserved.
+
+#pragma once
+
+#include "GraphEditorDragDropAction.h"
+
+namespace UE::SceneState::Graph
+{
+	struct FBlueprintAction_Graph;
+}
+
+namespace UE::SceneState::Editor
+{
+
+/** Drag Drop Op for State Machine Menu Entries */
+class FStateMachineItemDragDropOp : public FGraphSchemaActionDragDropAction
+{
+public:
+	using Super = FGraphSchemaActionDragDropAction;
+
+	DRAG_DROP_OPERATOR_TYPE(FStateMachineItemDragDropOp, FGraphSchemaActionDragDropAction)
+
+	static TSharedRef<FGraphSchemaActionDragDropAction> New(const TSharedPtr<FEdGraphSchemaAction>& InActionNode);
+
+	//~ Begin FGraphEditorDragDropAction
+	virtual void HoverTargetChanged() override;
+	virtual FReply DroppedOnCategory(FText InCategory) override;
+	virtual FReply DroppedOnAction(TSharedRef<FEdGraphSchemaAction> InAction) override;
+	//~ End FGraphEditorDragDropAction
+
+private:
+	/** Called when hovering over a category */
+	void OnHoverCategory(const FText& InCategory);
+
+	/** Called when hovering over a target action */
+	void OnHoverTargetAction(const TSharedRef<FEdGraphSchemaAction>& InTargetAction);
+
+	void SetFeedbackMessageOk(const FText& InMessage);
+	void SetFeedbackMessageError(const FText& InMessage);
+};
+
+} // UE::SceneState::Editor
diff --git a/Engine/Plugins/Experimental/SceneState/Source/SceneStateMachineGraph/Private/SceneStateMachineGraph.cpp b/Engine/Plugins/Experimental/SceneState/Source/SceneStateMachineGraph/Private/SceneStateMachineGraph.cpp
--- a/Engine/Plugins/Experimental/SceneState/Source/SceneStateMachineGraph/Private/SceneStateMachineGraph.cpp
+++ b/Engine/Plugins/Experimental/SceneState/Source/SceneStateMachineGraph/Private/SceneStateMachineGraph.cpp
@@ -1,8 +1,11 @@
 // Copyright Epic Games, Inc. All Rights Reserved.
 
 #include "SceneStateMachineGraph.h"
+#include "Engine/Blueprint.h"
 #include "Nodes/SceneStateMachineEntryNode.h"
+#include "Nodes/SceneStateMachineStateNode.h"
 #include "SceneStateBindingUtils.h"
+#include "Kismet2/BlueprintEditorUtils.h"
 
 USceneStateMachineGraph::FOnParametersChanged USceneStateMachineGraph::OnParametersChangedDelegate;
 
@@ -46,6 +49,16 @@ USceneStateMachineEntryNode* USceneStateMachineGraph::GetEntryNode() const
 	return LastEntryNode;
 }
 
+USceneStateMachineStateNode* USceneStateMachineGraph::GetParentStateNode() const
+{
+	return Cast<USceneStateMachineStateNode>(GetOuter());
+}
+
+bool USceneStateMachineGraph::IsRootStateMachine() const
+{
+	return !!Cast<UBlueprint>(GetOuter());
+}
+
 void USceneStateMachineGraph::AddNode(UEdGraphNode* InNodeToAdd, bool bInUserAction, bool bInSelectNewNode)
 {
 	if (!InNodeToAdd)
@@ -97,6 +110,19 @@ void USceneStateMachineGraph::PostEditImport()
 	GenerateNewParametersId();
 }
 
+void USceneStateMachineGraph::PostEditChangeProperty(FPropertyChangedEvent& InPropertyChangedEvent)
+{
+	Super::PostEditChangeProperty(InPropertyChangedEvent);
+
+	if (InPropertyChangedEvent.GetMemberPropertyName() == GET_MEMBER_NAME_CHECKED(USceneStateMachineGraph, Category))
+	{
+		if (UBlueprint* const Blueprint = FBlueprintEditorUtils::FindBlueprintForGraph(this))
+		{
+			FBlueprintEditorUtils::MarkBlueprintAsModified(Blueprint);
+		}
+	}
+}
+
 void USceneStateMachineGraph::GenerateNewParametersId()
 {
 	const FGuid OldParametersId = ParametersId;
diff --git a/Engine/Plugins/Experimental/SceneState/Source/SceneStateMachineGraph/Public/SceneStateMachineGraph.h b/Engine/Plugins/Experimental/SceneState/Source/SceneStateMachineGraph/Public/SceneStateMachineGraph.h
--- a/Engine/Plugins/Experimental/SceneState/Source/SceneStateMachineGraph/Public/SceneStateMachineGraph.h
+++ b/Engine/Plugins/Experimental/SceneState/Source/SceneStateMachineGraph/Public/SceneStateMachineGraph.h
@@ -7,7 +7,10 @@
 #include "StructUtils/PropertyBag.h"
 #include "SceneStateMachineGraph.generated.h"
 
+#define UE_API SCENESTATEMACHINEGRAPH_API
+
 class USceneStateMachineEntryNode;
+class USceneStateMachineStateNode;
 
 UCLASS(MinimalAPI, HideCategories=(Graph))
 class USceneStateMachineGraph : public UEdGraph
@@ -19,37 +22,50 @@ public:
 
 	USceneStateMachineGraph();
 
-	SCENESTATEMACHINEGRAPH_API static FOnParametersChanged::RegistrationType& OnParametersChanged();
+	UE_API static FOnParametersChanged::RegistrationType& OnParametersChanged();
+
+	UE_API void NotifyParametersChanged();
+
+	UE_API USceneStateMachineEntryNode* GetEntryNode() const;
 
-	SCENESTATEMACHINEGRAPH_API void NotifyParametersChanged();
+	/** Returns the parent node if it's a state machine nested in a state node*/
+	UE_API USceneStateMachineStateNode* GetParentStateNode() const;
 
-	SCENESTATEMACHINEGRAPH_API USceneStateMachineEntryNode* GetEntryNode() const;
+	/** Returns whether the state machine is at the root level (i.e. directly under the blueprint and not nested under a state node) */
+	UE_API bool IsRootStateMachine() const;
 
 	//~ Begin UEdGraph
-	SCENESTATEMACHINEGRAPH_API virtual void AddNode(UEdGraphNode* InNodeToAdd, bool bInUserAction, bool bInSelectNewNode) override;
+	UE_API virtual void AddNode(UEdGraphNode* InNodeToAdd, bool bInUserAction, bool bInSelectNewNode) override;
 	//~ End UEdGraph
 
 	//~ Begin UObject
-	SCENESTATEMACHINEGRAPH_API virtual void PostInitProperties() override;
-	SCENESTATEMACHINEGRAPH_API virtual void PostLoad() override;
-	SCENESTATEMACHINEGRAPH_API virtual void PostDuplicate(bool bInDuplicateForPIE) override;
-	SCENESTATEMACHINEGRAPH_API virtual void PostEditImport() override;
+	UE_API virtual void PostInitProperties() override;
+	UE_API virtual void PostLoad() override;
+	UE_API virtual void PostDuplicate(bool bInDuplicateForPIE) override;
+	UE_API virtual void PostEditImport() override;
+	UE_API virtual void PostEditChangeProperty(FPropertyChangedEvent& InPropertyChangedEvent) override;
 	//~ End UObject
 
 	/** Called to set a new unique id for parameters (e.g. after duplicating) */
 	void GenerateNewParametersId();
 
+	/** Category of the State Machine */
+	UPROPERTY(EditAnywhere, Category="State Machine")
+	FText Category;
+
 	/** Identifier for the Parameters Struct Id */
-	UPROPERTY(VisibleAnywhere, Category="Scene State")
+	UPROPERTY(VisibleAnywhere, Category="State Machine")
 	FGuid ParametersId;
 
-	UPROPERTY(EditAnywhere, Category="Scene State")
+	UPROPERTY(EditAnywhere, Category="State Machine")
 	FInstancedPropertyBag Parameters;
 
 	/** The run-mode for the State Machine. Currently only applies to Top-Level State Machines */
-	UPROPERTY(EditAnywhere, Category="Scene State")
+	UPROPERTY(EditAnywhere, Category="State Machine")
 	ESceneStateMachineRunMode RunMode = ESceneStateMachineRunMode::Auto;
 
 private:
 	static FOnParametersChanged OnParametersChangedDelegate;
 };
+
+#undef UE_API
