Motion Design: Transition parameters on Conduits crash fix

diff --git a/Engine/Plugins/Experimental/SceneState/Source/SceneState/Private/Conduit/SceneStateConduit.cpp b/Engine/Plugins/Experimental/SceneState/Source/SceneState/Private/Conduit/SceneStateConduit.cpp
--- a/Engine/Plugins/Experimental/SceneState/Source/SceneState/Private/Conduit/SceneStateConduit.cpp
+++ b/Engine/Plugins/Experimental/SceneState/Source/SceneState/Private/Conduit/SceneStateConduit.cpp
@@ -1,6 +1,7 @@
 // Copyright Epic Games, Inc. All Rights Reserved.
 
 #include "Conduit/SceneStateConduit.h"
+#include "Conduit/SceneStateConduitInstance.h"
 #include "Conduit/SceneStateConduitLink.h"
 #include "SceneStateExecutionContext.h"
 #include "SceneStateGeneratedClass.h"
@@ -31,6 +32,22 @@ void FSceneStateConduit::Link(const FSceneStateConduitLink& InConduitLink, UScen
 	}
 }
 
+void FSceneStateConduit::Setup(const FSceneStateExecutionContext& InContext) const
+{
+	FSceneStateConduitInstance* Instance = InContext.FindOrAddConduitInstance(*this);
+	if (!Instance || Instance->bInitialized)
+	{
+		return;
+	}
+
+	Instance->bInitialized = true;
+
+	for (const FSceneStateTransition& Transition : InContext.GetTransitions(*this))
+	{
+		Transition.Setup(InContext);
+	}
+}
+
 bool FSceneStateConduit::Evaluate(const UE::SceneState::FTransitionEvaluationParams& InParams) const
 {
 	// Early return if waiting for tasks to finish and there are still active tasks yet to finish
@@ -43,6 +60,11 @@ bool FSceneStateConduit::Evaluate(const UE::SceneState::FTransitionEvaluationPar
 	return ProcessEvaluationEvent(InParams);
 }
 
+void FSceneStateConduit::Exit(const FSceneStateExecutionContext& InContext) const
+{
+	InContext.RemoveConduitInstance(*this);
+}
+
 bool FSceneStateConduit::ProcessEvaluationEvent(const UE::SceneState::FTransitionEvaluationParams& InParams) const
 {
 	if (EnumHasAllFlags(EvaluationFlags, ESceneStateTransitionEvaluationFlags::EvaluationEventAlwaysTrue))
diff --git a/Engine/Plugins/Experimental/SceneState/Source/SceneState/Private/SceneStateExecutionContext.cpp b/Engine/Plugins/Experimental/SceneState/Source/SceneState/Private/SceneStateExecutionContext.cpp
--- a/Engine/Plugins/Experimental/SceneState/Source/SceneState/Private/SceneStateExecutionContext.cpp
+++ b/Engine/Plugins/Experimental/SceneState/Source/SceneState/Private/SceneStateExecutionContext.cpp
@@ -288,7 +288,7 @@ FInstancedPropertyBag FSceneStateExecutionContext::GetTemplateTransitionParamete
 	uint16 TransitionIndex;
 	if (GetTransitionIndex(InTransition, TransitionIndex))
 	{
-		return GeneratedClass->GetTransitionParameter(TransitionIndex);
+		return GeneratedClass->GetTransitionParameters(TransitionIndex);
 	}
 	return FInstancedPropertyBag();
 }
@@ -373,6 +373,34 @@ const FSceneStateMachine* FSceneStateExecutionContext::GetStateMachine(const FGu
 	return nullptr;
 }
 
+const FSceneStateMachine* FSceneStateExecutionContext::GetStateMachine(uint16 InAbsoluteIndex) const
+{
+	if (!GeneratedClass)
+	{
+		return nullptr;
+	}
+
+	const TConstArrayView<FSceneStateMachine> StateMachinesView = GeneratedClass->GetStateMachines();
+
+	checkfSlow(StateMachinesView.IsValidIndex(InAbsoluteIndex)
+		, TEXT("[Context: %s] Absolute Index %d is out of bounds! (Total State Machine Count: %d)")
+		, *GetExecutionContextName()
+		, InAbsoluteIndex
+		, StateMachinesView.Num());
+
+	return &StateMachinesView[InAbsoluteIndex];
+}
+
+const FSceneStateMachine* FSceneStateExecutionContext::GetContextStateMachine() const
+{
+	if (StateMachineExecutionStack.IsEmpty() || !GeneratedClass)
+	{
+		return nullptr;
+	}
+
+	return GetStateMachine(StateMachineExecutionStack.Last());
+}
+
 const FSceneState* FSceneStateExecutionContext::GetActiveState(const FSceneStateMachine& InStateMachine) const
 {
 	if (const FSceneStateMachineInstance* Instance = FindStateMachineInstance(InStateMachine))
@@ -660,6 +688,25 @@ void FSceneStateExecutionContext::RemoveTransitionInstance(const FSceneStateTran
 	}
 }
 
+FSceneStateConduitInstance* FSceneStateExecutionContext::FindOrAddConduitInstance(const FSceneStateConduit& InConduit) const
+{
+	uint16 ConduitIndex;
+	if (GetConduitIndex(InConduit, /*out*/ConduitIndex))
+	{
+		return &ConduitInstances.FindOrAdd(ConduitIndex);
+	}
+	return nullptr;
+}
+
+void FSceneStateExecutionContext::RemoveConduitInstance(const FSceneStateConduit& InConduit) const
+{
+	uint16 ConduitIndex;
+	if (GetConduitIndex(InConduit, /*out*/ConduitIndex))
+	{
+		ConduitInstances.Remove(ConduitIndex);
+	}
+}
+
 const FSceneStateExecutionContext* FSceneStateExecutionContext::FindParentContext() const
 {
 	if (TSharedPtr<UE::SceneState::FExecutionContextRegistry> ContextRegistry = ContextRegistryWeak.Pin())
@@ -707,3 +754,16 @@ bool FSceneStateExecutionContext::GetTransitionIndex(const FSceneStateTransition
 	}
 	return false;
 }
+
+bool FSceneStateExecutionContext::GetConduitIndex(const FSceneStateConduit& InConduit, uint16& OutIndex) const
+{
+	if (GeneratedClass)
+	{
+		TConstArrayView<FSceneStateConduit> ConduitsView = GeneratedClass->GetConduits();
+		const int32 Index = static_cast<int32>(&InConduit - ConduitsView.GetData());
+		checkSlow(ConduitsView.IsValidIndex(Index) && Index >= 0 && Index <= TNumericLimits<uint16>::Max());
+		OutIndex = static_cast<uint16>(Index);
+		return true;
+	}
+	return false;
+}
diff --git a/Engine/Plugins/Experimental/SceneState/Source/SceneState/Private/SceneStateExecutionScope.cpp b/Engine/Plugins/Experimental/SceneState/Source/SceneState/Private/SceneStateExecutionScope.cpp
new file mode 100644
--- /dev/null
+++ b/Engine/Plugins/Experimental/SceneState/Source/SceneState/Private/SceneStateExecutionScope.cpp
@@ -0,0 +1,22 @@
+// Copyright Epic Games, Inc. All Rights Reserved.
+
+#include "SceneStateExecutionScope.h"
+#include "SceneStateExecutionContext.h"
+
+namespace UE::SceneState
+{
+
+FExecutionScope::FExecutionScope(const FSceneStateExecutionContext& InContext, const FSceneStateMachine& InStateMachine)
+	: ExecutionContext(InContext)
+{
+	ensure(ExecutionContext.GetStateMachineIndex(InStateMachine, StateMachineIndex));
+	ExecutionContext.StateMachineExecutionStack.Push(StateMachineIndex);
+}
+
+FExecutionScope::~FExecutionScope()
+{
+	const uint16 RemovedElement = ExecutionContext.StateMachineExecutionStack.Pop(EAllowShrinking::No);
+	ensureMsgf(RemovedElement == StateMachineIndex, TEXT("State Machine Execution Stack expected to pop %d but instead popped %d"), StateMachineIndex, RemovedElement);
+}
+
+} // UE::SceneState
diff --git a/Engine/Plugins/Experimental/SceneState/Source/SceneState/Private/SceneStateGeneratedClass.cpp b/Engine/Plugins/Experimental/SceneState/Source/SceneState/Private/SceneStateGeneratedClass.cpp
--- a/Engine/Plugins/Experimental/SceneState/Source/SceneState/Private/SceneStateGeneratedClass.cpp
+++ b/Engine/Plugins/Experimental/SceneState/Source/SceneState/Private/SceneStateGeneratedClass.cpp
@@ -50,7 +50,7 @@ const FSceneStateMachine* USceneStateGeneratedClass::FindStateMachine(const FGui
 	return nullptr;
 }
 
-FInstancedPropertyBag USceneStateGeneratedClass::GetTransitionParameter(uint16 InAbsoluteIndex) const
+FInstancedPropertyBag USceneStateGeneratedClass::GetTransitionParameters(uint16 InAbsoluteIndex) const
 {
 	return TransitionParameters.FindRef(InAbsoluteIndex);
 }
diff --git a/Engine/Plugins/Experimental/SceneState/Source/SceneState/Private/SceneStateMachine.cpp b/Engine/Plugins/Experimental/SceneState/Source/SceneState/Private/SceneStateMachine.cpp
--- a/Engine/Plugins/Experimental/SceneState/Source/SceneState/Private/SceneStateMachine.cpp
+++ b/Engine/Plugins/Experimental/SceneState/Source/SceneState/Private/SceneStateMachine.cpp
@@ -5,6 +5,7 @@
 #include "SceneState.h"
 #include "SceneStateBindingUtils.h"
 #include "SceneStateExecutionContext.h"
+#include "SceneStateExecutionScope.h"
 #include "SceneStateLog.h"
 #include "SceneStateMachineInstance.h"
 #include "SceneStateObject.h"
@@ -59,7 +60,12 @@ void FSceneStateMachine::Start(const FSceneStateExecutionContext& InContext) con
 	}
 
 	Instance->Status = EExecutionStatus::Running;
-	ActiveState->Enter(InContext);
+
+	{
+		const FExecutionScope ExecutionScope(InContext, *this);
+		ActiveState->Enter(InContext);
+	}
+
 	EvaluateStateTransitions(InContext, *Instance, *ActiveState);
 }
 
@@ -86,7 +92,10 @@ void FSceneStateMachine::Tick(const FSceneStateExecutionContext& InContext, floa
 		return;
 	}
 
-	ActiveState->Tick(InContext, InDeltaSeconds);
+	{
+		const FExecutionScope ExecutionScope(InContext, *this);
+		ActiveState->Tick(InContext, InDeltaSeconds);
+	}
 
 	EvaluateStateTransitions(InContext, *Instance, *ActiveState);
 }
@@ -127,6 +136,7 @@ void FSceneStateMachine::EvaluateStateTransitions(const FSceneStateExecutionCont
 				.SourceState = InActiveState
 			};
 
+		const FExecutionScope ExecutionScope(InContext, *this);
 		Transition = EvaluateTransitions(InContext.GetTransitions(InActiveState), EvaluationParams);
 	}
 
@@ -158,6 +168,7 @@ void FSceneStateMachine::EvaluateStateTransitions(const FSceneStateExecutionCont
 		return;
 	}
 
+	const FExecutionScope ExecutionScope(InContext, *this);
 	NewActiveState->Enter(InContext);
 }
 
@@ -245,6 +256,7 @@ void FSceneStateMachine::StopActiveState(const FSceneStateExecutionContext& InCo
 {
 	if (const FSceneState* ActiveState = InContext.GetActiveState(*this))
 	{
+		const UE::SceneState::FExecutionScope ExecutionScope(InContext, *this);
 		ActiveState->Exit(InContext);
 	}
 }
diff --git a/Engine/Plugins/Experimental/SceneState/Source/SceneState/Private/Transition/SceneStateTransition.cpp b/Engine/Plugins/Experimental/SceneState/Source/SceneState/Private/Transition/SceneStateTransition.cpp
--- a/Engine/Plugins/Experimental/SceneState/Source/SceneState/Private/Transition/SceneStateTransition.cpp
+++ b/Engine/Plugins/Experimental/SceneState/Source/SceneState/Private/Transition/SceneStateTransition.cpp
@@ -35,10 +35,19 @@ void FSceneStateTransition::Link(const FSceneStateTransitionLink& InTransitionLi
 
 void FSceneStateTransition::Setup(const FSceneStateExecutionContext& InContext) const
 {
-	if (FSceneStateTransitionInstance* Instance = InContext.FindOrAddTransitionInstance(*this))
+	FSceneStateTransitionInstance* Instance = InContext.FindOrAddTransitionInstance(*this);
+	if (!ensure(Instance))
 	{
-		Instance->Parameters = InContext.GetTemplateTransitionParameter(*this);
+		return;
 	}
+
+	Instance->Parameters = InContext.GetTemplateTransitionParameter(*this);
+
+	// Setup the conduit chaining this transition
+	// if (const FSceneStateConduit* Conduit = FindConduitTarget(InContext))
+	// {
+	// 	Conduit->Setup(InContext);
+	// }
 }
 
 bool FSceneStateTransition::Evaluate(const UE::SceneState::FTransitionEvaluationParams& InParams) const
@@ -62,6 +71,12 @@ bool FSceneStateTransition::Evaluate(const UE::SceneState::FTransitionEvaluation
 void FSceneStateTransition::Exit(const FSceneStateExecutionContext& InContext) const
 {
 	InContext.RemoveTransitionInstance(*this);
+
+	// Exit the conduit chaining this transition
+	if (const FSceneStateConduit* Conduit = FindConduitTarget(InContext))
+	{
+		Conduit->Exit(InContext);
+	}
 }
 
 bool FSceneStateTransition::ApplyBindings(const FSceneStateExecutionContext& InContext, FSceneStateTransitionInstance& InInstance) const
@@ -142,17 +157,23 @@ bool FSceneStateTransition::ProcessEvaluationEvent(const UE::SceneState::FTransi
 	if (EvaluationEvent->NumParms != 0)
 	{
 		FSceneStateTransitionInstance* Instance = InParams.ExecutionContext.FindOrAddTransitionInstance(*this);
-		if (Instance && Instance->Parameters.IsValid())
+		if (!ensureMsgf(Instance && Instance->Parameters.IsValid()
+			, TEXT("Transition Instance Parameters is not valid. Transition evaluation event expects %d parameters.")
+			, EvaluationEvent->NumParms))
 		{
-			ApplyBindings(InParams.ExecutionContext, *Instance);
-
-			checkfSlow(EvaluationEvent->NumParms == Instance->Parameters.GetNumPropertiesInBag()
-				, TEXT("Unexpected parameter mismatch! Event Parameter Count: %d... Transition Parameter Count: %d")
-				, EvaluationEvent->NumParms
-				, Instance->Parameters.GetNumPropertiesInBag());
+			return false;
+		}
 
-			FunctionParams = Instance->Parameters.GetMutableValue().GetMemory();
+		if (!ensureMsgf(EvaluationEvent->NumParms == Instance->Parameters.GetNumPropertiesInBag()
+			, TEXT("Unexpected parameter mismatch! Event Parameter Count: %d... Transition Parameter Count: %d")
+			, EvaluationEvent->NumParms
+			, Instance->Parameters.GetNumPropertiesInBag()))
+		{
+			return false;
 		}
+
+		ApplyBindings(InParams.ExecutionContext, *Instance);
+		FunctionParams = Instance->Parameters.GetMutableValue().GetMemory();
 	}
 
 	USceneStateObject* RootState = InParams.ExecutionContext.GetRootState();
@@ -163,3 +184,26 @@ bool FSceneStateTransition::ProcessEvaluationEvent(const UE::SceneState::FTransi
 	check(Result);
 	return Result->bCanTransition;
 }
+
+const FSceneStateConduit* FSceneStateTransition::FindConduitTarget(const FSceneStateExecutionContext& InContext) const
+{
+	if (Target.Type != ESceneStateTransitionTargetType::Conduit)
+	{
+		// Target isn't a conduit, skip
+		return nullptr;
+	}
+
+	const FSceneStateMachine* StateMachine = InContext.GetContextStateMachine();
+	if (!ensureMsgf(StateMachine, TEXT("Unexpectedly found no context state machine. Is this being directly called outside of a State Machine?")))
+	{
+		return nullptr;
+	}
+
+	const FSceneStateConduit* Conduit = InContext.GetConduit(*StateMachine, Target.Index);
+	if (!ensureMsgf(Conduit, TEXT("Unable to find conduit at target index %d. State Machine Conduit count: %d"), Target.Index, StateMachine->GetConduitRange().Count))
+	{
+		return nullptr;
+	}
+
+	return Conduit;
+}
diff --git a/Engine/Plugins/Experimental/SceneState/Source/SceneState/Public/Conduit/SceneStateConduit.h b/Engine/Plugins/Experimental/SceneState/Source/SceneState/Public/Conduit/SceneStateConduit.h
--- a/Engine/Plugins/Experimental/SceneState/Source/SceneState/Public/Conduit/SceneStateConduit.h
+++ b/Engine/Plugins/Experimental/SceneState/Source/SceneState/Public/Conduit/SceneStateConduit.h
@@ -9,6 +9,7 @@
 
 class USceneStateGeneratedClass;
 struct FSceneStateConduitLink;
+struct FSceneStateExecutionContext;
 
 namespace UE::SceneState
 {
@@ -30,9 +31,15 @@ struct FSceneStateConduit
 	/** Called when the Generated Class is linking to cache the Event function and Result property */
 	void Link(const FSceneStateConduitLink& InConduitLink, USceneStateGeneratedClass* InGeneratedClass);
 
-	/** Called to evaluate this transition. Returns true if the transition should take place */
+	/** Called when a state that could evaluate this conduit has entered. Allocates and initializes the conduit instance */
+	void Setup(const FSceneStateExecutionContext& InContext) const;
+
+	/** Called to evaluate this conduit. Returns true if the conduit should take place */
 	bool Evaluate(const UE::SceneState::FTransitionEvaluationParams& InParams) const;
 
+	/** Called when the State has exited. Cleans up the conduit instance */
+	void Exit(const FSceneStateExecutionContext& InContext) const;
+
 	FSceneStateRange GetTransitionRange() const
 	{
 		return TransitionRange;
diff --git a/Engine/Plugins/Experimental/SceneState/Source/SceneState/Public/Conduit/SceneStateConduitInstance.h b/Engine/Plugins/Experimental/SceneState/Source/SceneState/Public/Conduit/SceneStateConduitInstance.h
new file mode 100644
--- /dev/null
+++ b/Engine/Plugins/Experimental/SceneState/Source/SceneState/Public/Conduit/SceneStateConduitInstance.h
@@ -0,0 +1,16 @@
+// Copyright Epic Games, Inc. All Rights Reserved.
+
+#pragma once
+
+#include "StructUtils/PropertyBag.h"
+#include "SceneStateConduitInstance.generated.h"
+
+/** Instance data of a conduit */
+USTRUCT()
+struct FSceneStateConduitInstance
+{
+	GENERATED_BODY()
+
+	UPROPERTY()
+	bool bInitialized = false;
+};
diff --git a/Engine/Plugins/Experimental/SceneState/Source/SceneState/Public/SceneStateExecutionContext.h b/Engine/Plugins/Experimental/SceneState/Source/SceneState/Public/SceneStateExecutionContext.h
--- a/Engine/Plugins/Experimental/SceneState/Source/SceneState/Public/SceneStateExecutionContext.h
+++ b/Engine/Plugins/Experimental/SceneState/Source/SceneState/Public/SceneStateExecutionContext.h
@@ -2,6 +2,7 @@
 
 #pragma once
 
+#include "Conduit/SceneStateConduitInstance.h"
 #include "Containers/Array.h"
 #include "SceneStateExecutionContextHandle.h"
 #include "SceneStateInstance.h"
@@ -33,6 +34,7 @@ struct FSceneStateTransition;
 namespace UE::SceneState
 {
 	class FExecutionContextRegistry;
+	class FExecutionScope;
 }
 
 /**
@@ -138,6 +140,12 @@ struct FSceneStateExecutionContext
 	/** Returns the State Machine linked to the given id */
 	const FSceneStateMachine* GetStateMachine(const FGuid& InStateMachineId) const;
 
+	/** Returns the State Machine at the given absolute index */
+	const FSceneStateMachine* GetStateMachine(uint16 InAbsoluteIndex) const;
+
+	/** Returns the most relevant State Machine in the current scope, if any */
+	const FSceneStateMachine* GetContextStateMachine() const;
+
 	/** Returns the currently Active State within this context for the given State Machine */
 	const FSceneState* GetActiveState(const FSceneStateMachine& InStateMachine) const;
 
@@ -207,6 +215,12 @@ struct FSceneStateExecutionContext
 	/** Removes the transition instance of this context for the given transition */
 	void RemoveTransitionInstance(const FSceneStateTransition& InTransition) const;
 
+	/** Returns the conduit instance of this context for the given conduit, adding a new one if it's not found */
+	FSceneStateConduitInstance* FindOrAddConduitInstance(const FSceneStateConduit& InConduit) const;
+
+	/** Removes the conduit instance of this context for the given conduit */
+	void RemoveConduitInstance(const FSceneStateConduit& InConduit) const;
+
 private:
 	/** Retrieves the parent execution context if existing. Can be null */
 	const FSceneStateExecutionContext* FindParentContext() const;
@@ -220,6 +234,9 @@ private:
 	/** Retrieves the Absolute Index of a Transition in the Transitions array */
 	bool GetTransitionIndex(const FSceneStateTransition& InTransition, uint16& OutIndex) const;
 
+	/** Retrieves the Absolute Index of a Conduit in the Conduits array */
+	bool GetConduitIndex(const FSceneStateConduit& InConduit, uint16& OutIndex) const;
+
 	/** Root state object owning the Scene State Execution */
 	UPROPERTY()
 	TObjectPtr<USceneStateObject> RootState = nullptr;
@@ -249,6 +266,9 @@ private:
 	UPROPERTY()
 	mutable TMap<uint16, FSceneStateMachineInstance> StateMachineInstances;
 
+	/** Array of state machine absolute indices from outermost to innermost for this execution context */
+	mutable TArray<uint16, TInlineAllocator<4>> StateMachineExecutionStack;
+
 	/**
 	 * Map of the transition absolute index to its instance data
 	 * It's allocated when the state starts (along with the other exit transitions) and removed on state exit
@@ -256,9 +276,18 @@ private:
 	UPROPERTY()
 	mutable TMap<uint16, FSceneStateTransitionInstance> TransitionInstances;
 
+	/**
+	 * Map of the conduit absolute index to its instance data
+	 * It's allocated when a state that could evaluate the conduit starts and removed on state exit
+	 */
+	UPROPERTY()
+	mutable TMap<uint16, FSceneStateConduitInstance> ConduitInstances;
+
 	/** Context handle to the parent execution context. Can be invalid. */
 	UE::SceneState::FExecutionContextHandle ParentHandle;
 
 	/** Weak reference to the registry containing this context */
 	TWeakPtr<UE::SceneState::FExecutionContextRegistry> ContextRegistryWeak;
+
+	friend class UE::SceneState::FExecutionScope;
 };
diff --git a/Engine/Plugins/Experimental/SceneState/Source/SceneState/Public/SceneStateExecutionScope.h b/Engine/Plugins/Experimental/SceneState/Source/SceneState/Public/SceneStateExecutionScope.h
new file mode 100644
--- /dev/null
+++ b/Engine/Plugins/Experimental/SceneState/Source/SceneState/Public/SceneStateExecutionScope.h
@@ -0,0 +1,34 @@
+// Copyright Epic Games, Inc. All Rights Reserved.
+
+#pragma once
+
+#include "HAL/Platform.h"
+#include <limits>
+
+struct FSceneStateExecutionContext;
+struct FSceneStateMachine;
+
+namespace UE::SceneState
+{
+
+/**
+ * Adds the given State Machine to the Execution Stack on scope enter / constructor
+ * and removes it on scope exit / destructor.
+ * Stores references and should be short-lived within a function's scope.
+ */
+class FExecutionScope
+{
+public:
+	explicit FExecutionScope(const FSceneStateExecutionContext& InContext, const FSceneStateMachine& InStateMachine);
+
+	~FExecutionScope();
+
+private:
+	/** The execution context containing the execution stack */
+	const FSceneStateExecutionContext& ExecutionContext;
+
+	/** Index to the state machine in the execution scope */
+	uint16 StateMachineIndex = std::numeric_limits<uint16>::max();
+};
+
+} // UE::SceneState
diff --git a/Engine/Plugins/Experimental/SceneState/Source/SceneState/Public/SceneStateGeneratedClass.h b/Engine/Plugins/Experimental/SceneState/Source/SceneState/Public/SceneStateGeneratedClass.h
--- a/Engine/Plugins/Experimental/SceneState/Source/SceneState/Public/SceneStateGeneratedClass.h
+++ b/Engine/Plugins/Experimental/SceneState/Source/SceneState/Public/SceneStateGeneratedClass.h
@@ -89,7 +89,7 @@ public:
 		return Transitions;
 	}
 
-	FInstancedPropertyBag GetTransitionParameter(uint16 InAbsoluteIndex) const;
+	FInstancedPropertyBag GetTransitionParameters(uint16 InAbsoluteIndex) const;
 
 	TConstArrayView<FSceneStateMachine> GetStateMachines() const
 	{
diff --git a/Engine/Plugins/Experimental/SceneState/Source/SceneState/Public/Transition/SceneStateTransition.h b/Engine/Plugins/Experimental/SceneState/Source/SceneState/Public/Transition/SceneStateTransition.h
--- a/Engine/Plugins/Experimental/SceneState/Source/SceneState/Public/Transition/SceneStateTransition.h
+++ b/Engine/Plugins/Experimental/SceneState/Source/SceneState/Public/Transition/SceneStateTransition.h
@@ -15,6 +15,7 @@ class UFunction;
 class UObject;
 class USceneStateGeneratedClass;
 struct FSceneState;
+struct FSceneStateConduit;
 struct FSceneStateExecutionContext;
 struct FSceneStateMachine;
 struct FSceneStateMachineInstance;
@@ -64,6 +65,9 @@ private:
 	/** Calls the evaluation event (can return early if optimized away via the evaluation flags) */
 	bool ProcessEvaluationEvent(const UE::SceneState::FTransitionEvaluationParams& InParams) const;
 
+	/** If the transition target is a conduit, returns the pointer to such conduit */
+	const FSceneStateConduit* FindConduitTarget(const FSceneStateExecutionContext& InContext) const;
+
 	/** Bindings Batch where this Transition is target */
 	UPROPERTY()
 	FPropertyBindingIndex16 BindingsBatch;
