Motion Design Scene State: uncommented code block

diff --git a/Engine/Plugins/Experimental/SceneState/Source/SceneState/Private/Transition/SceneStateTransition.cpp b/Engine/Plugins/Experimental/SceneState/Source/SceneState/Private/Transition/SceneStateTransition.cpp
--- a/Engine/Plugins/Experimental/SceneState/Source/SceneState/Private/Transition/SceneStateTransition.cpp
+++ b/Engine/Plugins/Experimental/SceneState/Source/SceneState/Private/Transition/SceneStateTransition.cpp
@@ -44,10 +44,10 @@ void FSceneStateTransition::Setup(const FSceneStateExecutionContext& InContext)
 	Instance->Parameters = InContext.GetTemplateTransitionParameter(*this);
 
 	// Setup the conduit chaining this transition
-	// if (const FSceneStateConduit* Conduit = FindConduitTarget(InContext))
-	// {
-	// 	Conduit->Setup(InContext);
-	// }
+	if (const FSceneStateConduit* Conduit = FindConduitTarget(InContext))
+	{
+		Conduit->Setup(InContext);
+	}
 }
 
 bool FSceneStateTransition::Evaluate(const UE::SceneState::FTransitionEvaluationParams& InParams) const
