Motion Design Scene State: fix potential issue where iteration limit could be reached and return true while result being invalid

diff --git a/Engine/Plugins/Experimental/SceneState/Source/SceneState/Private/SceneStatePropertyUtils.cpp b/Engine/Plugins/Experimental/SceneState/Source/SceneState/Private/SceneStatePropertyUtils.cpp
--- a/Engine/Plugins/Experimental/SceneState/Source/SceneState/Private/SceneStatePropertyUtils.cpp
+++ b/Engine/Plugins/Experimental/SceneState/Source/SceneState/Private/SceneStatePropertyUtils.cpp
@@ -56,9 +56,15 @@ bool ResolveProperty(const FSceneStateExecutionContext& InContext, const FSceneS
 	constexpr int32 MaxIterations = 100;
 	int32 IterationCount = 0;
 
-	// If the Source is a Property Reference, follow the reference chain until a non-property reference is found (or a safe max iteration limit is hit)
-	while (ensure(IterationCount++ < MaxIterations) && IsPropertyReference(Result.ResolvedReference->SourceLeafProperty))
+	// If the Source is a Property Reference, follow the reference chain until a non-property reference is found
+	while (IsPropertyReference(Result.ResolvedReference->SourceLeafProperty))
 	{
+		// Safe max iteration limit hit. Could this be a reference loop?
+		if (!ensureMsgf(++IterationCount <= MaxIterations, TEXT("Resolve Property Reference failed as it hit the iteration limit. Could this be a reference loop?")))
+		{
+			return false;
+		}
+
 		const FSceneStatePropertyReference& SourcePropertyReference = *reinterpret_cast<FSceneStatePropertyReference*>(Result.ValuePtr);
 		if (!Private::ResolvePropertyInternal(InContext, SourcePropertyReference, Result))
 		{
