Motion Design Scene State: fix issue where unexposed properties (i.e. non-instance-editable properties) would still show up in Blueprint Tasks + Player tasks.

diff --git a/Engine/Plugins/Experimental/SceneState/Source/SceneStateBlueprintEditor/Private/SceneStateBlueprintEditorUtils.cpp b/Engine/Plugins/Experimental/SceneState/Source/SceneStateBlueprintEditor/Private/SceneStateBlueprintEditorUtils.cpp
--- a/Engine/Plugins/Experimental/SceneState/Source/SceneStateBlueprintEditor/Private/SceneStateBlueprintEditorUtils.cpp
+++ b/Engine/Plugins/Experimental/SceneState/Source/SceneStateBlueprintEditor/Private/SceneStateBlueprintEditorUtils.cpp
@@ -139,7 +139,7 @@ UClass* FindCommonBase(const TSharedRef<IPropertyHandle>& InPropertyHandle)
 	return CommonBase;
 }
 
-bool AddObjectProperties(const TSharedRef<IPropertyHandle>& InPropertyHandle, IDetailChildrenBuilder& InChildBuilder)
+bool AddObjectProperties(const TSharedRef<IPropertyHandle>& InPropertyHandle, IDetailChildrenBuilder& InChildBuilder, uint64 InDisallowedFlags)
 {
 	const UClass* ObjectClass = FindCommonBase(InPropertyHandle);
 	if (!ObjectClass)
@@ -151,6 +151,11 @@ bool AddObjectProperties(const TSharedRef<IPropertyHandle>& InPropertyHandle, ID
 
 	for (const FProperty* Property : TFieldRange<FProperty>(ObjectClass, EFieldIterationFlags::IncludeSuper))
 	{
+		if (Property->HasAnyPropertyFlags(InDisallowedFlags))
+		{
+			continue;
+		}
+
 		TSharedPtr<IPropertyHandle> PropertyHandle = InPropertyHandle->GetChildHandle(Property->GetFName());
 		if (!PropertyHandle.IsValid() || !PropertyHandle->IsValidHandle())
 		{
diff --git a/Engine/Plugins/Experimental/SceneState/Source/SceneStateBlueprintEditor/Public/SceneStateBlueprintEditorUtils.h b/Engine/Plugins/Experimental/SceneState/Source/SceneStateBlueprintEditor/Public/SceneStateBlueprintEditorUtils.h
--- a/Engine/Plugins/Experimental/SceneState/Source/SceneStateBlueprintEditor/Public/SceneStateBlueprintEditorUtils.h
+++ b/Engine/Plugins/Experimental/SceneState/Source/SceneStateBlueprintEditor/Public/SceneStateBlueprintEditorUtils.h
@@ -61,7 +61,12 @@ namespace UE::SceneState::Editor
 	/** Finds the Common Class for all the Objects within the Object Property Handle */
 	SCENESTATEBLUEPRINTEDITOR_API UClass* FindCommonBase(const TSharedRef<IPropertyHandle>& InPropertyHandle);
 
-	/** Helper function to add the Common properties of the object property handle to the Child Builder */
-	SCENESTATEBLUEPRINTEDITOR_API bool AddObjectProperties(const TSharedRef<IPropertyHandle>& InPropertyHandle, IDetailChildrenBuilder& InChildBuilder);
+	/**
+	 * Helper function to add the Common properties of the object property handle to the Child Builder
+	 * @param InPropertyHandle the object property handle to get the common properties from
+	 * @param InChildBuilder the builder to add the properties to
+	 * @param InDisallowedFlags filter flags. Any property with any of these flags will be skipped 
+	 */
+	SCENESTATEBLUEPRINTEDITOR_API bool AddObjectProperties(const TSharedRef<IPropertyHandle>& InPropertyHandle, IDetailChildrenBuilder& InChildBuilder, uint64 InDisallowedFlags = CPF_DisableEditOnInstance);
 
 } // UE::SceneState::Editor
