Minor comment update

diff --git a/Engine/Plugins/Experimental/SceneState/Source/SceneStateBinding/Public/SceneStatePropertyReference.h b/Engine/Plugins/Experimental/SceneState/Source/SceneStateBinding/Public/SceneStatePropertyReference.h
--- a/Engine/Plugins/Experimental/SceneState/Source/SceneStateBinding/Public/SceneStatePropertyReference.h
+++ b/Engine/Plugins/Experimental/SceneState/Source/SceneStateBinding/Public/SceneStatePropertyReference.h
@@ -19,7 +19,7 @@
  * Specifies the type of property that the reference supports.
  * Supported types are: bool, byte, int32, int64, float, double, Name, String, Text, Wildcard, UObject pointers, and structs.
  * Wildcards is a 'bypass' so that any valid FProperty is accepted, and is up to implementer to handle this.
- * Structs and Objects must use full path name.
+ * Structs and Objects must use full path name. (e.g. for Actor type "/Script/Engine.Actor")
  *
  * IsRefToArray
  * Specified to mean that the reference is to an TArray<RefType>
