Motion Design Scene State: moved reset to be called from PurgeClass rather than manually from the compiler context. The compiler context already calls PurgeClass in its process of cleaning the class.

diff --git a/Engine/Plugins/Experimental/SceneState/Source/SceneState/Private/SceneStateGeneratedClass.cpp b/Engine/Plugins/Experimental/SceneState/Source/SceneState/Private/SceneStateGeneratedClass.cpp
--- a/Engine/Plugins/Experimental/SceneState/Source/SceneState/Private/SceneStateGeneratedClass.cpp
+++ b/Engine/Plugins/Experimental/SceneState/Source/SceneState/Private/SceneStateGeneratedClass.cpp
@@ -79,6 +79,12 @@ void USceneStateGeneratedClass::Link(FArchive& Ar, bool bInRelinkExistingPropert
 	}
 }
 
+void USceneStateGeneratedClass::PurgeClass(bool bInRecompilingOnLoad)
+{
+	Super::PurgeClass(bInRecompilingOnLoad);
+	Reset();
+}
+
 void USceneStateGeneratedClass::PostLoad()
 {
 	Super::PostLoad();
diff --git a/Engine/Plugins/Experimental/SceneState/Source/SceneState/Public/SceneStateGeneratedClass.h b/Engine/Plugins/Experimental/SceneState/Source/SceneState/Public/SceneStateGeneratedClass.h
--- a/Engine/Plugins/Experimental/SceneState/Source/SceneState/Public/SceneStateGeneratedClass.h
+++ b/Engine/Plugins/Experimental/SceneState/Source/SceneState/Public/SceneStateGeneratedClass.h
@@ -110,6 +110,7 @@ public:
 
 	//~ Begin UClass
 	SCENESTATE_API virtual void Link(FArchive& Ar, bool bInRelinkExistingProperties) override;
+	SCENESTATE_API virtual void PurgeClass(bool bInRecompilingOnLoad) override;
 	//~ End UClass
 
 	//~ Begin UObject
@@ -153,7 +154,7 @@ private:
 	SCENESTATE_API void ResolveBindings();
 
 	/** Resets all the elements that get compiled for this generated class (e.g. Tasks, State Machines, etc) */
-	SCENESTATE_API void Reset();
+	void Reset();
 
 	bool IsFullClass() const;
 
diff --git a/Engine/Plugins/Experimental/SceneState/Source/SceneStateBlueprintEditor/Private/Compiler/SceneStateBlueprintCompilerContext.cpp b/Engine/Plugins/Experimental/SceneState/Source/SceneStateBlueprintEditor/Private/Compiler/SceneStateBlueprintCompilerContext.cpp
--- a/Engine/Plugins/Experimental/SceneState/Source/SceneStateBlueprintEditor/Private/Compiler/SceneStateBlueprintCompilerContext.cpp
+++ b/Engine/Plugins/Experimental/SceneState/Source/SceneStateBlueprintEditor/Private/Compiler/SceneStateBlueprintCompilerContext.cpp
@@ -97,16 +97,6 @@ void FBlueprintCompilerContext::EnsureProperGeneratedClass(UClass*& InOutTargetC
 	}
 }
 
-void FBlueprintCompilerContext::CleanAndSanitizeClass(UBlueprintGeneratedClass* InClassToClean, UObject*& InOutOldCDO)
-{
-	FKismetCompilerContext::CleanAndSanitizeClass(InClassToClean, InOutOldCDO);
-
-	if (USceneStateGeneratedClass* GeneratedClass = Cast<USceneStateGeneratedClass>(InClassToClean))
-	{
-		GeneratedClass->Reset();
-	}
-}
-
 void FBlueprintCompilerContext::SaveSubObjectsFromCleanAndSanitizeClass(FSubobjectCollection& OutSubObjectsToSave, UBlueprintGeneratedClass* InClassToClean)
 {
 	FKismetCompilerContext::SaveSubObjectsFromCleanAndSanitizeClass(OutSubObjectsToSave, InClassToClean);
diff --git a/Engine/Plugins/Experimental/SceneState/Source/SceneStateBlueprintEditor/Private/Compiler/SceneStateBlueprintCompilerContext.h b/Engine/Plugins/Experimental/SceneState/Source/SceneStateBlueprintEditor/Private/Compiler/SceneStateBlueprintCompilerContext.h
--- a/Engine/Plugins/Experimental/SceneState/Source/SceneStateBlueprintEditor/Private/Compiler/SceneStateBlueprintCompilerContext.h
+++ b/Engine/Plugins/Experimental/SceneState/Source/SceneStateBlueprintEditor/Private/Compiler/SceneStateBlueprintCompilerContext.h
@@ -22,7 +22,6 @@ private:
 	virtual void SpawnNewClass(const FString& InNewClassName) override;
 	virtual void OnNewClassSet(UBlueprintGeneratedClass* InClassToUse) override;
 	virtual void EnsureProperGeneratedClass(UClass*& InOutTargetClass) override;
-	virtual void CleanAndSanitizeClass(UBlueprintGeneratedClass* InClassToClean, UObject*& InOutOldCDO) override;
 	virtual void SaveSubObjectsFromCleanAndSanitizeClass(FSubobjectCollection& OutSubObjectsToSave, UBlueprintGeneratedClass* InClassToClean) override;
 	virtual void MergeUbergraphPagesIn(UEdGraph* InUbergraph) override;
 	virtual void OnPostCDOCompiled(const UObject::FPostCDOCompiledContext& InContext) override;
