Motion Design Scene State: added actor spawner task that is able to hold onto a 'model' actor that will then be used as a template to spawn the actual actor in game.
This also improves FSceneState::InstanceTaskObjects to support nested object instanced references.

diff --git a/Engine/Plugins/Experimental/SceneState/Source/SceneState/Private/SceneState.cpp b/Engine/Plugins/Experimental/SceneState/Source/SceneState/Private/SceneState.cpp
--- a/Engine/Plugins/Experimental/SceneState/Source/SceneState/Private/SceneState.cpp
+++ b/Engine/Plugins/Experimental/SceneState/Source/SceneState/Private/SceneState.cpp
@@ -270,25 +270,45 @@ void FSceneState::InstanceTaskObjects(UObject* InOuter, TConstArrayView<FStructV
 			continue;
 		}
 
-		for (FObjectProperty* ObjectProperty : TFieldRange<FObjectProperty>(Source.GetScriptStruct()))
-		{
-			if (!ObjectProperty->HasAnyPropertyFlags(CPF_InstancedReference))
+		// Visit all object properties that are instanced references, stepping into structs but not object properties themselves.
+		// Using const_cast here as there's no Visit variant taking a const ptr.
+		Source.GetScriptStruct()->Visit(const_cast<uint8*>(Source.GetMemory()),
+			[&InDuplicationFunc, InOuter, Target](const FPropertyVisitorContext& InContext)->EPropertyVisitorControlFlow
 			{
-				continue;
-			}
+				const FObjectProperty* ObjectProperty = CastField<FObjectProperty>(InContext.Path.Top().Property);
+				if (!ObjectProperty)
+				{
+					return EPropertyVisitorControlFlow::StepInto;
+				}
 
-			UObject* InstanceObject = nullptr;
-			if (UObject* TemplateObject = ObjectProperty->GetObjectPropertyValue_InContainer(Source.GetMemory()))
-			{
-				FObjectDuplicationParameters Parameters = InitStaticDuplicateObjectParams(TemplateObject, InOuter);
-				Parameters.DestName = MakeUniqueObjectName(Parameters.DestOuter, Parameters.SourceObject->GetClass(), Parameters.SourceObject->GetFName());
-				Parameters.FlagMask = RF_AllFlags & ~RF_DefaultSubObject;
-				Parameters.PortFlags |= PPF_DuplicateVerbatim; // Skip resetting text IDs
+				if (!ObjectProperty->HasAnyPropertyFlags(CPF_InstancedReference))
+				{
+					return EPropertyVisitorControlFlow::StepOver;
+				}
 
-				InstanceObject = InDuplicationFunc(Parameters);
-			}
-			ObjectProperty->SetPropertyValue_InContainer(Target.GetMemory(), InstanceObject);
-		}
+				const void* const SourcePropertyMemory = InContext.Data.PropertyData;
+
+				UObject* InstanceObject = nullptr;
+				if (UObject* TemplateObject = ObjectProperty->GetObjectPropertyValue(SourcePropertyMemory))
+				{
+					FObjectDuplicationParameters Parameters = InitStaticDuplicateObjectParams(TemplateObject, InOuter);
+					Parameters.DestName = MakeUniqueObjectName(Parameters.DestOuter, Parameters.SourceObject->GetClass(), Parameters.SourceObject->GetFName());
+					Parameters.FlagMask = RF_AllFlags & ~RF_DefaultSubObject;
+					Parameters.PortFlags |= PPF_DuplicateVerbatim; // Skip resetting text IDs
+
+					InstanceObject = InDuplicationFunc(Parameters);
+
+					UE_LOG(LogSceneState, Verbose, TEXT("Instantiated object '%s' (Outer: '%s') from template '%s' (Outer: '%s')")
+						, *GetNameSafe(InstanceObject)
+						, *InOuter->GetName()
+						, *TemplateObject->GetName()
+						, *GetNameSafe(TemplateObject->GetOuter()));
+				}
+
+				void* const TargetPropertyMemory = UE::SceneState::ResolveVisitedPath(Target.GetScriptStruct(), Target.GetMemory(), InContext.Path);
+				ObjectProperty->SetPropertyValue(TargetPropertyMemory, InstanceObject);
+				return EPropertyVisitorControlFlow::StepOver;
+			});
 	}
 }
 
diff --git a/Engine/Plugins/Experimental/SceneState/Source/SceneState/Private/SceneStateUtils.cpp b/Engine/Plugins/Experimental/SceneState/Source/SceneState/Private/SceneStateUtils.cpp
--- a/Engine/Plugins/Experimental/SceneState/Source/SceneState/Private/SceneStateUtils.cpp
+++ b/Engine/Plugins/Experimental/SceneState/Source/SceneState/Private/SceneStateUtils.cpp
@@ -8,6 +8,8 @@
 #include "UObject/Class.h"
 #include "UObject/Object.h"
 #include "UObject/Package.h"
+#include "UObject/PropertyVisitor.h"
+#include "UObject/UnrealType.h"
 
 namespace UE::SceneState
 {
@@ -61,6 +63,28 @@ TArray<FConstStructView> GetConstStructViews(const FInstancedStructContainer& In
 	return GetConstStructViews(InStructContainer, Range);
 }
 
+void* ResolveVisitedPath(const UScriptStruct* InRootObject, void* InRootData, const FPropertyVisitorPath& InPath)
+{
+	if (!InRootObject || !InRootData)
+	{
+		return nullptr;
+	}
+
+	const TArray<FPropertyVisitorInfo>& PathArray = InPath.GetPath();
+	if (PathArray.IsEmpty())
+	{
+		return nullptr;
+	}
+
+	void* PropertyData = InRootObject->ResolveVisitedPathInfo(InRootData, PathArray[0]);
+	for (int32 PathIndex = 1; PropertyData && PathIndex < PathArray.Num(); ++PathIndex)
+	{
+		const FPropertyVisitorInfo& PreviousInfo = PathArray[PathIndex - 1];
+		PropertyData = PreviousInfo.Property->ResolveVisitedPathInfo(PropertyData, PathArray[PathIndex]);
+	}
+	return PropertyData;
+}
+
 void DiscardObject(UObject* InObjectToDiscard)
 {
 	if (InObjectToDiscard)
diff --git a/Engine/Plugins/Experimental/SceneState/Source/SceneState/Private/Tasks/SceneStateTask.cpp b/Engine/Plugins/Experimental/SceneState/Source/SceneState/Private/Tasks/SceneStateTask.cpp
--- a/Engine/Plugins/Experimental/SceneState/Source/SceneState/Private/Tasks/SceneStateTask.cpp
+++ b/Engine/Plugins/Experimental/SceneState/Source/SceneState/Private/Tasks/SceneStateTask.cpp
@@ -41,6 +41,14 @@ void FSceneStateTask::BuildTaskInstance(UObject* InOuter, FStructView InTaskInst
 		OnBuildTaskInstance(InOuter, InTaskInstance);
 	}
 }
+
+void FSceneStateTask::PostEditChange(UE::SceneState::FTaskEditChange& InEditChange, FStructView InTaskInstance)
+{
+	if (InTaskInstance.IsValid())
+	{
+		OnPostEditChange(InEditChange, InTaskInstance);
+	}
+}
 #endif
 
 const FSceneStateTaskBindingExtension* FSceneStateTask::GetBindingExtension() const
diff --git a/Engine/Plugins/Experimental/SceneState/Source/SceneState/Public/SceneStateEnums.h b/Engine/Plugins/Experimental/SceneState/Source/SceneState/Public/SceneStateEnums.h
--- a/Engine/Plugins/Experimental/SceneState/Source/SceneState/Public/SceneStateEnums.h
+++ b/Engine/Plugins/Experimental/SceneState/Source/SceneState/Public/SceneStateEnums.h
@@ -23,6 +23,16 @@ enum class EExecutionStatus : uint8
 	Finished,
 };
 
+/** Describes the different components of a task */
+enum class ETaskObjectType : uint8
+{
+	None,
+	/** Object deriving from FSceneStateTask that contains the logic to run */
+	Task,
+	/** Object deriving from FSceneStateTaskInstance containing the instance data to feed to the task */
+	TaskInstance,
+};
+
 } // UE::SceneState
 
 /** Defines common task behaviors: whether it should tick, whether it handles custom bindings, etc. */
diff --git a/Engine/Plugins/Experimental/SceneState/Source/SceneState/Public/SceneStateUtils.h b/Engine/Plugins/Experimental/SceneState/Source/SceneState/Public/SceneStateUtils.h
--- a/Engine/Plugins/Experimental/SceneState/Source/SceneState/Public/SceneStateUtils.h
+++ b/Engine/Plugins/Experimental/SceneState/Source/SceneState/Public/SceneStateUtils.h
@@ -10,6 +10,7 @@ class UClass;
 class UObject;
 struct FConstStructView;
 struct FInstancedStructContainer;
+struct FPropertyVisitorPath;
 struct FSceneStateRange;
 struct FStructView;
 template<typename T> struct TObjectPtr;
@@ -38,6 +39,14 @@ namespace UE::SceneState
 
 	SCENESTATE_API TArray<FConstStructView> GetConstStructViews(const FInstancedStructContainer& InStructContainer);
 
+	/**
+	 * Given a FPropertyVisitorPath, attempt to resolve that to a valid data pointer.
+	 * RootObject is required to implement ResolveVisitedPathInfo to provide the resolver logic.
+	 * TODO: FPropertyVisitorPath::GetPropertyDataPtr only takes a UObject ptr.
+	 * @see FPropertyVisitorPath::GetPropertyDataPtr
+	 */
+	SCENESTATE_API void* ResolveVisitedPath(const UScriptStruct* InRootObject, void* InRootData, const FPropertyVisitorPath& InPath);
+
 	/**
 	 * Discards a given object by changing its outer to transient package and adding a prefix to it while also ensuring uniqueness
 	 * @param InObjectToDiscard the object to discard
diff --git a/Engine/Plugins/Experimental/SceneState/Source/SceneState/Public/Tasks/SceneStateTask.h b/Engine/Plugins/Experimental/SceneState/Source/SceneState/Public/Tasks/SceneStateTask.h
--- a/Engine/Plugins/Experimental/SceneState/Source/SceneState/Public/Tasks/SceneStateTask.h
+++ b/Engine/Plugins/Experimental/SceneState/Source/SceneState/Public/Tasks/SceneStateTask.h
@@ -12,12 +12,17 @@ struct FSceneStateExecutionContext;
 struct FSceneStateTaskBindingExtension;
 struct FSceneStateTaskInstance;
 
-namespace UE::SceneState::Editor
+namespace UE::SceneState
 {
-	class FBindingCompiler;
-	class FBlueprintCompilerContext;
-	class FStateMachineCompiler;
-	class FStateMachineTaskCompiler;
+	struct FTaskEditChange;
+
+	namespace Editor
+	{
+		class FBindingCompiler;
+		class FBlueprintCompilerContext;
+		class FStateMachineCompiler;
+		class FStateMachineTaskCompiler;
+	}
 }
 
 USTRUCT(meta=(Hidden))
@@ -40,6 +45,13 @@ struct FSceneStateTask
 	 * @param InTaskInstance the task instance to build
 	 */
 	SCENESTATE_API void BuildTaskInstance(UObject* InOuter, FStructView InTaskInstance) const;
+
+	/**
+	 * Called when a property within a task has changed
+	 * @param InEditChange the edit change information (property changed event, outer)
+	 * @param InTaskInstance the task instance to build
+	 */
+	SCENESTATE_API void PostEditChange(UE::SceneState::FTaskEditChange& InEditChange, FStructView InTaskInstance);
 #endif
 
 	/** Called to get the optional Binding Extension for a Task */
@@ -96,6 +108,14 @@ protected:
 	virtual void OnBuildTaskInstance(UObject* InOuter, FStructView InTaskInstance) const
 	{
 	}
+
+	/**
+	 * Called when a property within a task has changed
+	 * @see PostEditChange
+	 */
+	virtual void OnPostEditChange(UE::SceneState::FTaskEditChange& InEditChange, FStructView InTaskInstance) const
+	{
+	}
 #endif
 
 	/**
diff --git a/Engine/Plugins/Experimental/SceneState/Source/SceneState/Public/Tasks/SceneStateTaskEditChange.h b/Engine/Plugins/Experimental/SceneState/Source/SceneState/Public/Tasks/SceneStateTaskEditChange.h
new file mode 100644
--- /dev/null
+++ b/Engine/Plugins/Experimental/SceneState/Source/SceneState/Public/Tasks/SceneStateTaskEditChange.h
@@ -0,0 +1,27 @@
+// Copyright Epic Games, Inc. All Rights Reserved.
+
+#pragma once
+
+#include "SceneStateEnums.h"
+#include "UObject/UnrealType.h"
+
+class UObject;
+
+namespace UE::SceneState
+{
+#if WITH_EDITOR
+/** Describes an Edit change for a Task */
+struct FTaskEditChange : FPropertyChangedEvent
+{
+	FTaskEditChange()
+		: FPropertyChangedEvent(nullptr)
+	{
+	}
+	/** Owning object of the task */
+	UObject* Outer = nullptr;
+	/** The object that was changed: whether it was the task or the task instance */
+	ETaskObjectType ChangedObject = ETaskObjectType::None;
+};
+#endif
+
+} // UE::SceneState
diff --git a/Engine/Plugins/Experimental/SceneState/Source/SceneStateBlueprintEditor/Private/SceneStateBlueprintEditorStyle.cpp b/Engine/Plugins/Experimental/SceneState/Source/SceneStateBlueprintEditor/Private/SceneStateBlueprintEditorStyle.cpp
--- a/Engine/Plugins/Experimental/SceneState/Source/SceneStateBlueprintEditor/Private/SceneStateBlueprintEditorStyle.cpp
+++ b/Engine/Plugins/Experimental/SceneState/Source/SceneStateBlueprintEditor/Private/SceneStateBlueprintEditorStyle.cpp
@@ -35,7 +35,8 @@ FBlueprintEditorStyle::FBlueprintEditorStyle()
 	Set("TaskIcon.SceneStateBlueprintableTaskWrapper", new IMAGE_BRUSH_SVG("Starship/AssetIcons/Blueprint_16", CoreStyleConstants::Icon16x16));
 	Set("TaskIcon.SceneStateDelayTask", new IMAGE_BRUSH_SVG("Starship/Common/Timecode", CoreStyleConstants::Icon16x16));
 	Set("TaskIcon.SceneStatePrintStringTask", new CORE_IMAGE_BRUSH_SVG("Starship/Common/OutputLog", CoreStyleConstants::Icon16x16));
-	
+	Set("TaskIcon.SceneStateSpawnActorTask", new IMAGE_BRUSH_SVG("Starship/AssetIcons/Actor_16", CoreStyleConstants::Icon16x16));
+
 	FSlateStyleRegistry::RegisterSlateStyle(*this);
 }
 
diff --git a/Engine/Plugins/Experimental/SceneState/Source/SceneStateMachineGraph/Private/Nodes/SceneStateMachineTaskNode.cpp b/Engine/Plugins/Experimental/SceneState/Source/SceneStateMachineGraph/Private/Nodes/SceneStateMachineTaskNode.cpp
--- a/Engine/Plugins/Experimental/SceneState/Source/SceneStateMachineGraph/Private/Nodes/SceneStateMachineTaskNode.cpp
+++ b/Engine/Plugins/Experimental/SceneState/Source/SceneStateMachineGraph/Private/Nodes/SceneStateMachineTaskNode.cpp
@@ -10,6 +10,7 @@
 #include "Tasks/SceneStateTaskBindingExtension.h"
 #include "Tasks/SceneStateTaskDesc.h"
 #include "Tasks/SceneStateTaskDescRegistry.h"
+#include "Tasks/SceneStateTaskEditChange.h"
 
 #define LOCTEXT_NAMESPACE "SceneStateMachineTaskNode"
 
@@ -34,7 +35,7 @@ bool USceneStateMachineTaskNode::SetTaskStruct(const UScriptStruct* InTaskStruct
 		{
 			Task.Reset();
 		}
-		OnTaskTypeChanged();
+		OnTaskChanged();
 		return true;
 	}
 	return false;
@@ -55,7 +56,7 @@ void USceneStateMachineTaskNode::SetTaskBlueprintClass(TSubclassOf<USceneStateBl
 
 	if (bTaskChanged)
 	{
-		OnTaskTypeChanged();
+		OnTaskChanged();
 	}
 }
 
@@ -172,15 +173,26 @@ void USceneStateMachineTaskNode::BeginDestroy()
 	UE::SceneState::OnStructIdChanged.RemoveAll(this);
 }
 
-void USceneStateMachineTaskNode::PostEditChangeProperty(FPropertyChangedEvent& InPropertyChangedEvent)
+void USceneStateMachineTaskNode::PostEditChangeChainProperty(FPropertyChangedChainEvent& InPropertyChangedEvent)
 {
-	Super::PostEditChangeProperty(InPropertyChangedEvent);
+	using namespace UE::SceneState;
+
+	Super::PostEditChangeChainProperty(InPropertyChangedEvent);
 
 	CachedTaskDisplayName = FText::GetEmpty();
 
-	if (InPropertyChangedEvent.GetMemberPropertyName() == GET_MEMBER_NAME_CHECKED(USceneStateMachineTaskNode, Task))
+	const ETaskObjectType ChangedObject = GetChangedObject(InPropertyChangedEvent);
+
+	if (ChangedObject != ETaskObjectType::None)
 	{
-		OnTaskTypeChanged();
+		OnTaskChanged();
+
+		// If the change chain is deeper than just the task, it means that the property is within the task or task instance
+		// so propagate this edit change event to the task
+		if (InPropertyChangedEvent.PropertyChain.Num() > 1)
+		{
+			PostEditChangeTask(InPropertyChangedEvent, ChangedObject);
+		}
 	}
 
 	OnPostEditTaskDelegate.Broadcast();
@@ -204,6 +216,76 @@ void USceneStateMachineTaskNode::PostEditImport()
 	RebuildTaskInstance();
 }
 
+UE::SceneState::ETaskObjectType USceneStateMachineTaskNode::GetChangedObject(FPropertyChangedChainEvent& InPropertyChangedEvent)
+{
+	using namespace UE::SceneState;
+
+	const TDoubleLinkedList<FProperty*>::TDoubleLinkedListNode* const MemberPropertyNode = InPropertyChangedEvent.PropertyChain.GetHead();
+	if (!MemberPropertyNode)
+	{
+		return ETaskObjectType::None;
+	}
+
+	const FProperty* const MemberProperty = MemberPropertyNode->GetValue();
+	if (!MemberProperty)
+	{
+		return ETaskObjectType::None;
+	}
+
+	const FName MemberPropertyName = MemberProperty->GetFName();
+
+	if (MemberPropertyName == GET_MEMBER_NAME_CHECKED(USceneStateMachineTaskNode, Task))
+	{
+		return ETaskObjectType::Task;
+	}
+
+	if (MemberPropertyName == GET_MEMBER_NAME_CHECKED(USceneStateMachineTaskNode, TaskInstance))
+	{
+		return ETaskObjectType::TaskInstance;
+	}
+
+	return ETaskObjectType::None;
+}
+
+void USceneStateMachineTaskNode::PostEditChangeTask(FPropertyChangedChainEvent& InPropertyChangedEvent, UE::SceneState::ETaskObjectType InChangedObject)
+{
+	if (!TaskInstance.IsValid())
+	{
+		return;
+	}
+
+	FSceneStateTask* const ResolvedTask = Task.GetMutablePtr<FSceneStateTask>();
+	if (!ResolvedTask)
+	{
+		return;
+	}
+
+	const TDoubleLinkedList<FProperty*>::TDoubleLinkedListNode* MemberPropertyNode = InPropertyChangedEvent.PropertyChain.GetHead();
+	const TDoubleLinkedList<FProperty*>::TDoubleLinkedListNode* ChangedPropertyNode = InPropertyChangedEvent.PropertyChain.GetTail();
+	if (!MemberPropertyNode || !ChangedPropertyNode)
+	{
+		return;
+	}
+
+	// Move the member property as it should now be 1 level down
+	MemberPropertyNode = MemberPropertyNode->GetNextNode();
+	if (!MemberPropertyNode)
+	{
+		return;
+	}
+
+	UE::SceneState::FTaskEditChange TaskEditChange;
+	TaskEditChange.ChangeType = InPropertyChangedEvent.ChangeType;
+	TaskEditChange.MemberProperty = MemberPropertyNode->GetValue();
+	TaskEditChange.Property = ChangedPropertyNode->GetValue();
+
+	TaskEditChange.Outer = this;
+	TaskEditChange.ChangedObject = InChangedObject;
+
+	FStructView TaskInstanceView(TaskInstance.GetScriptStruct(), TaskInstance.GetMutableMemory());
+	ResolvedTask->PostEditChange(TaskEditChange, TaskInstanceView);
+}
+
 FSceneStateTaskDescContext USceneStateMachineTaskNode::MakeContext() const
 {
 	FSceneStateTaskDescContext Context;
@@ -231,7 +313,7 @@ void USceneStateMachineTaskNode::RebuildTaskInstance()
 	}
 }
 
-void USceneStateMachineTaskNode::OnTaskTypeChanged()
+void USceneStateMachineTaskNode::OnTaskChanged()
 {
 	if (const FSceneStateTask* ResolvedTask = Task.GetPtr<FSceneStateTask>())
 	{
diff --git a/Engine/Plugins/Experimental/SceneState/Source/SceneStateMachineGraph/Public/Nodes/SceneStateMachineTaskNode.h b/Engine/Plugins/Experimental/SceneState/Source/SceneStateMachineGraph/Public/Nodes/SceneStateMachineTaskNode.h
--- a/Engine/Plugins/Experimental/SceneState/Source/SceneStateMachineGraph/Public/Nodes/SceneStateMachineTaskNode.h
+++ b/Engine/Plugins/Experimental/SceneState/Source/SceneStateMachineGraph/Public/Nodes/SceneStateMachineTaskNode.h
@@ -14,6 +14,7 @@ struct FSceneStateTaskDescMutableContext;
 
 namespace UE::SceneState
 {
+	enum class ETaskObjectType : uint8;
 	struct FStructIdChange;
 }
 
@@ -65,7 +66,7 @@ public:
 	//~ Begin UObject
 	virtual void PostLoad() override;
 	virtual void BeginDestroy() override;
-	virtual void PostEditChangeProperty(FPropertyChangedEvent& InPropertyChangedEvent) override;
+	virtual void PostEditChangeChainProperty(FPropertyChangedChainEvent& InPropertyChangedEvent) override;
 	virtual void PostDuplicate(bool bInDuplicateForPIE) override;
 	virtual void PostEditImport() override;
 	//~ End UObject
@@ -81,6 +82,12 @@ public:
 	}
 
 private:
+	/** Determines whether the changed property comes from the task or task instance object, or none of them */
+	static UE::SceneState::ETaskObjectType GetChangedObject(FPropertyChangedChainEvent& InPropertyChangedEvent);
+
+	/** Called on PostEditChangeChainProperty when the property that changed came from a Task object */
+	void PostEditChangeTask(FPropertyChangedChainEvent& InPropertyChangedEvent, UE::SceneState::ETaskObjectType InChangedObject);
+
 	/** Create context for this task node */
 	FSceneStateTaskDescContext MakeContext() const;
 
@@ -91,7 +98,7 @@ private:
 	void RebuildTaskInstance();
 
 	/** Called when the task type has changed */
-	void OnTaskTypeChanged();
+	void OnTaskChanged();
 
 	/** Called when a struct id has changed */
 	void OnStructIdChanged(const UE::SceneState::FStructIdChange& InChange);
diff --git a/Engine/Plugins/Experimental/SceneState/Source/SceneStateTasks/Private/SceneStateSpawnActorTask.cpp b/Engine/Plugins/Experimental/SceneState/Source/SceneStateTasks/Private/SceneStateSpawnActorTask.cpp
new file mode 100644
--- /dev/null
+++ b/Engine/Plugins/Experimental/SceneState/Source/SceneStateTasks/Private/SceneStateSpawnActorTask.cpp
@@ -0,0 +1,157 @@
+// Copyright Epic Games, Inc. All Rights Reserved.
+
+#include "SceneStateSpawnActorTask.h"
+#include "Engine/Engine.h"
+#include "GameFramework/Actor.h"
+#include "SceneStateExecutionContext.h"
+#include "SceneStateTasksLog.h"
+
+#if WITH_EDITOR
+#include "ActorFactories/ActorFactory.h"
+#include "AssetSelection.h"
+#include "Tasks/SceneStateTaskEditChange.h"
+#endif
+
+#if WITH_EDITOR
+const UScriptStruct* FSceneStateSpawnActorTask::OnGetTaskInstanceType() const
+{
+	return FInstanceDataType::StaticStruct();
+}
+
+void FSceneStateSpawnActorTask::OnPostEditChange(UE::SceneState::FTaskEditChange& InEditChange, FStructView InTaskInstance) const
+{
+	using namespace UE::SceneState;
+
+	if (InEditChange.ChangedObject == ETaskObjectType::Task
+		&& InEditChange.GetMemberPropertyName() == GET_MEMBER_NAME_CHECKED(FSceneStateSpawnActorTask, ActorClass))
+	{
+		UpdateActorTemplate(InEditChange.Outer, InTaskInstance);
+	}
+}
+#endif
+
+void FSceneStateSpawnActorTask::OnStart(const FSceneStateExecutionContext& InContext, FStructView InTaskInstance) const
+{
+	ON_SCOPE_EXIT
+	{
+		Finish(InContext, InTaskInstance);
+	};
+
+	FInstanceDataType& Instance = InTaskInstance.Get<FInstanceDataType>();
+
+	AActor* const TemplateActor = Instance.ActorTemplate.Template;
+	if (!TemplateActor)
+	{
+		UE_LOG(LogSceneStateTasks, Error, TEXT("[%s] Template Actor was not spawned. Template Actor is null."), *InContext.GetExecutionContextName());
+		return;
+	}
+
+	FText ErrorMessage;
+	if (!ShouldSpawnActor(InTaskInstance, ErrorMessage))
+	{
+		UE_LOG(LogSceneStateTasks, Error, TEXT("[%s] Template Actor was not spawned. Reason: %s")
+			, *InContext.GetExecutionContextName()
+			, *ErrorMessage.ToString());
+		return;
+	}
+
+	if (!GEngine)
+	{
+		UE_LOG(LogSceneStateTasks, Error, TEXT("[%s] Template Actor was not spawned. GEngine unexpectedly null."), *InContext.GetExecutionContextName());
+		return;
+	}
+
+	UWorld* const World = GEngine->GetWorldFromContextObject(InContext.GetContextObject(), EGetWorldErrorMode::LogAndReturnNull);
+	if (!World)
+	{
+		UE_LOG(LogSceneStateTasks, Error, TEXT("[%s] Template Actor was not spawned. Could not find a valid world for context object '%s'.")
+			, *InContext.GetExecutionContextName()
+			, *GetNameSafe(InContext.GetContextObject()));
+		return;
+	}
+
+	FActorSpawnParameters SpawnParameters;
+	SpawnParameters.Template = TemplateActor;
+	SpawnParameters.SpawnCollisionHandlingOverride = Instance.SpawnCollisionHandling;
+	SpawnParameters.ObjectFlags = RF_Transient | RF_Transactional;
+
+	if (AActor* Actor = World->SpawnActor(TemplateActor->GetClass(), &Instance.SpawnTransform, SpawnParameters))
+	{
+		OnActorSpawned(Actor, InTaskInstance);
+	}
+}
+
+#if WITH_EDITOR
+void FSceneStateSpawnActorTask::UpdateActorTemplate(UObject* InOuter, FStructView InTaskInstance) const
+{
+	FInstanceDataType& Instance = InTaskInstance.Get<FInstanceDataType>();
+
+	TObjectPtr<AActor>& TemplateActor = Instance.ActorTemplate.Template;
+	if (TemplateActor && TemplateActor->GetClass() == ActorClass && TemplateActor->GetOuter() == InOuter)
+	{
+		return;
+	}
+
+	if (!ActorClass)
+	{
+		TemplateActor = nullptr;
+		return;
+	}
+
+	// Use an existing world to create a temporary actor instance
+	UWorld* const World = GWorld;
+	if (!World)
+	{
+		TemplateActor = nullptr;
+		UE_LOG(LogSceneStateTasks, Error, TEXT("Unable to create an actor of class %s. No valid world found."), *ActorClass->GetName());
+		return;
+	}
+
+	UActorFactory* ActorFactory = FActorFactoryAssetProxy::GetFactoryForAssetObject(ActorClass);
+	if (!ActorFactory)
+	{
+		TemplateActor = nullptr;
+		UE_LOG(LogSceneStateTasks, Error, TEXT("Unable to find an actor factory for class %s."), *ActorClass->GetName());
+		return;
+	}
+
+	FText ErrorText;
+	if (!ActorFactory->CanCreateActorFrom(FAssetData(ActorClass), ErrorText))
+	{
+		TemplateActor = nullptr;
+		UE_LOG(LogSceneStateTasks, Error, TEXT("Unable to create actor of class %s. Reason: %s"), *ActorClass->GetName(), *ErrorText.ToString());
+		return;
+	}
+
+	FActorSpawnParameters SpawnParams;
+	SpawnParams.ObjectFlags = RF_Transient | RF_Transactional;
+
+	AActor* const SpawnedActor = ActorFactory->CreateActor(ActorClass, World->PersistentLevel, FTransform::Identity, SpawnParams);
+	if (!SpawnedActor)
+	{
+		TemplateActor = nullptr;
+		UE_LOG(LogSceneStateTasks, Error, TEXT("Unable to create actor of class %s. Actor Factory failed to create actor."), *ActorClass->GetName());
+		return;
+	}
+
+	constexpr bool bCallModify = false;
+	TemplateActor = Cast<AActor>(StaticDuplicateObject(SpawnedActor, InOuter, NAME_None, RF_AllFlags & ~RF_Transient));
+	TemplateActor->DetachFromActor(FDetachmentTransformRules(EDetachmentRule::KeepRelative, bCallModify));
+	TemplateActor->bIsEditorPreviewActor = false;
+
+	// Tasks and its instances containing objects are all outered to the Scene State Generated Class.
+	// BP Generated Classes are, by design, renamed so that they're outered to the package of the object owning the blueprint.
+	// See: UBlueprint::RenameGeneratedClasses.
+	// This means that actors outered to scene state generated classes will not have a valid ULevel
+	// as AActor::GetLevel() returns the first level object in the outer chain, which in scene state blueprints won't exist.
+	//
+	// UEditorEngine::Map_Load deletes all actors that are in the same package as a world but mismatch their Actor->GetWorld() with the given world.
+	// In the case of embedded scene state blueprints, the template actors would live on the same package, but not have a level in its outer chain.
+	// AActor::GetWorld relies on AActor::GetLevel to return a valid level. Because of this, template actors in embedded scene state bps would get deleted,
+	// and the only other condition that stops Map_Load from deleting these actors is if the actor is an archetype object.
+	TemplateActor->SetFlags(RF_ArchetypeObject);
+
+	constexpr bool bNetForce = false;
+	World->DestroyActor(SpawnedActor, bNetForce, bCallModify);
+}
+#endif
diff --git a/Engine/Plugins/Experimental/SceneState/Source/SceneStateTasks/Private/SceneStateTasksModule.cpp b/Engine/Plugins/Experimental/SceneState/Source/SceneStateTasks/Private/SceneStateTasksModule.cpp
--- a/Engine/Plugins/Experimental/SceneState/Source/SceneStateTasks/Private/SceneStateTasksModule.cpp
+++ b/Engine/Plugins/Experimental/SceneState/Source/SceneStateTasks/Private/SceneStateTasksModule.cpp
@@ -1,5 +1,8 @@
 // Copyright Epic Games, Inc. All Rights Reserved.
 
 #include "Modules/ModuleManager.h"
+#include "SceneStateTasksLog.h"
+
+DEFINE_LOG_CATEGORY(LogSceneStateTasks);
 
 IMPLEMENT_MODULE(FDefaultModuleImpl, SceneStateTasks)
diff --git a/Engine/Plugins/Experimental/SceneState/Source/SceneStateTasks/Public/SceneStateSpawnActorTask.h b/Engine/Plugins/Experimental/SceneState/Source/SceneStateTasks/Public/SceneStateSpawnActorTask.h
new file mode 100644
--- /dev/null
+++ b/Engine/Plugins/Experimental/SceneState/Source/SceneStateTasks/Public/SceneStateSpawnActorTask.h
@@ -0,0 +1,91 @@
+// Copyright Epic Games, Inc. All Rights Reserved.
+
+#pragma once
+
+#include "Engine/EngineTypes.h"
+#include "Tasks/SceneStateTask.h"
+#include "Tasks/SceneStateTaskInstance.h"
+#include "Templates/SubclassOf.h"
+#include "SceneStateSpawnActorTask.generated.h"
+
+#define UE_API SCENESTATETASKS_API
+
+struct FActorSpawnParameters;
+
+/** Represents the actor template to use */
+USTRUCT()
+struct FSceneStateActorTemplate
+{
+	GENERATED_BODY()
+
+	/** The actor template to use */
+	UPROPERTY(VisibleAnywhere, Instanced, Category="Scene State")
+	TObjectPtr<AActor> Template;
+};
+
+USTRUCT()
+struct FSceneStateSpawnActorTaskInstance : public FSceneStateTaskInstance
+{
+	GENERATED_BODY()
+
+	/** The actor template to use */
+	UPROPERTY(EditAnywhere, Category="Scene State", meta=(NoBindingSelfOnly))
+	FSceneStateActorTemplate ActorTemplate;
+
+	/** The spawn transform for the actor */
+	UPROPERTY(EditAnywhere, Category="Scene State")
+	FTransform SpawnTransform = FTransform::Identity;
+
+	/** Method for resolving collisions at the spawn point. Default (unspecified) means no override, use the actor's setting */
+	UPROPERTY(EditAnywhere, AdvancedDisplay, Category="Scene State")
+	ESpawnActorCollisionHandlingMethod SpawnCollisionHandling = ESpawnActorCollisionHandlingMethod::Undefined;
+};
+
+/** Spawns an actor of a given class */
+USTRUCT(DisplayName="Spawn Actor", Category="Core")
+struct FSceneStateSpawnActorTask : public FSceneStateTask
+{
+	GENERATED_BODY()
+
+	using FInstanceDataType = FSceneStateSpawnActorTaskInstance;
+
+protected:
+	//~ Begin FSceneStateTask
+#if WITH_EDITOR
+	UE_API virtual const UScriptStruct* OnGetTaskInstanceType() const override;
+	UE_API virtual void OnPostEditChange(UE::SceneState::FTaskEditChange& InEditChange, FStructView InTaskInstance) const override;
+#endif
+	UE_API virtual void OnStart(const FSceneStateExecutionContext& InContext, FStructView InTaskInstance) const override;
+	//~ End FSceneStateTask
+
+#if WITH_EDITOR
+	/** Called to update the actor template from the task instance to match the current actor class */
+	void UpdateActorTemplate(UObject* InOuter, FStructView InTaskInstance) const;
+#endif
+
+	/**
+	 * Determines whether the actor should be spawned
+	 * @param InTaskInstance data view to the task instance of the execution.
+	 * @param OutErrorMessage the error message indicating why an actor should not be spawned.
+	 * @return True if the actor should be spawned, false otherwise ideally accompanied by an error message. 
+	 */
+	virtual bool ShouldSpawnActor(FStructView InTaskInstance, FText& OutErrorMessage) const
+	{
+		return true;
+	}
+
+	/**
+	 * Called after the actor has spawned for additional handling
+	 * @param InActorChecked the spawned actor, already checked to be valid
+	 * @param InTaskInstance data view to the task instance of the execution.
+	 */
+	virtual void OnActorSpawned(AActor* InActorChecked, FStructView InTaskInstance) const
+	{
+	}
+
+	/** Actor class to spawn */
+	UPROPERTY(EditAnywhere, Category="Scene State")
+	TSubclassOf<AActor> ActorClass;
+};
+
+#undef UE_API
diff --git a/Engine/Plugins/Experimental/SceneState/Source/SceneStateTasks/Public/SceneStateTasksLog.h b/Engine/Plugins/Experimental/SceneState/Source/SceneStateTasks/Public/SceneStateTasksLog.h
new file mode 100644
--- /dev/null
+++ b/Engine/Plugins/Experimental/SceneState/Source/SceneStateTasks/Public/SceneStateTasksLog.h
@@ -0,0 +1,7 @@
+// Copyright Epic Games, Inc. All Rights Reserved.
+
+#pragma once
+
+#include "Logging/LogMacros.h"
+
+DECLARE_LOG_CATEGORY_EXTERN(LogSceneStateTasks, Log, All);
diff --git a/Engine/Plugins/Experimental/SceneState/Source/SceneStateTasks/SceneStateTasks.Build.cs b/Engine/Plugins/Experimental/SceneState/Source/SceneStateTasks/SceneStateTasks.Build.cs
--- a/Engine/Plugins/Experimental/SceneState/Source/SceneStateTasks/SceneStateTasks.Build.cs
+++ b/Engine/Plugins/Experimental/SceneState/Source/SceneStateTasks/SceneStateTasks.Build.cs
@@ -18,5 +18,15 @@ public class SceneStateTasks : ModuleRules
                 "SceneStateBinding",
             }
         );
+
+        if (Target.Type == TargetRules.TargetType.Editor)
+        {
+            PrivateDependencyModuleNames.AddRange(
+                new string[]
+                {
+                    "UnrealEd",
+                }
+            );
+        }
     }
 }
\ No newline at end of file
