Motion Design: fix post edit change comments for clarity that task instances are also considered.

diff --git a/Engine/Plugins/Experimental/SceneState/Source/SceneState/Public/Tasks/SceneStateTask.h b/Engine/Plugins/Experimental/SceneState/Source/SceneState/Public/Tasks/SceneStateTask.h
--- a/Engine/Plugins/Experimental/SceneState/Source/SceneState/Public/Tasks/SceneStateTask.h
+++ b/Engine/Plugins/Experimental/SceneState/Source/SceneState/Public/Tasks/SceneStateTask.h
@@ -47,9 +47,9 @@ struct FSceneStateTask
 	SCENESTATE_API void BuildTaskInstance(UObject* InOuter, FStructView InTaskInstance) const;
 
 	/**
-	 * Called when a property within a task has changed
+	 * Called when a property within a task or task instance has changed
 	 * @param InEditChange the edit change information (property changed event, outer)
-	 * @param InTaskInstance the task instance to build
+	 * @param InTaskInstance the task instance paired with the task
 	 */
 	SCENESTATE_API void PostEditChange(UE::SceneState::FTaskEditChange& InEditChange, FStructView InTaskInstance);
 #endif
diff --git a/Engine/Plugins/Experimental/SceneState/Source/SceneState/Public/Tasks/SceneStateTaskEditChange.h b/Engine/Plugins/Experimental/SceneState/Source/SceneState/Public/Tasks/SceneStateTaskEditChange.h
--- a/Engine/Plugins/Experimental/SceneState/Source/SceneState/Public/Tasks/SceneStateTaskEditChange.h
+++ b/Engine/Plugins/Experimental/SceneState/Source/SceneState/Public/Tasks/SceneStateTaskEditChange.h
@@ -10,7 +10,7 @@ class UObject;
 namespace UE::SceneState
 {
 #if WITH_EDITOR
-/** Describes an Edit change for a Task */
+/** Describes an Edit change for a task or task instance */
 struct FTaskEditChange : FPropertyChangedEvent
 {
 	FTaskEditChange()
diff --git a/Engine/Plugins/Experimental/SceneState/Source/SceneStateMachineGraph/Public/Nodes/SceneStateMachineTaskNode.h b/Engine/Plugins/Experimental/SceneState/Source/SceneStateMachineGraph/Public/Nodes/SceneStateMachineTaskNode.h
--- a/Engine/Plugins/Experimental/SceneState/Source/SceneStateMachineGraph/Public/Nodes/SceneStateMachineTaskNode.h
+++ b/Engine/Plugins/Experimental/SceneState/Source/SceneStateMachineGraph/Public/Nodes/SceneStateMachineTaskNode.h
@@ -85,7 +85,7 @@ private:
 	/** Determines whether the changed property comes from the task or task instance object, or none of them */
 	static UE::SceneState::ETaskObjectType GetChangedObject(FPropertyChangedChainEvent& InPropertyChangedEvent);
 
-	/** Called on PostEditChangeChainProperty when the property that changed came from a Task object */
+	/** Called on PostEditChangeChainProperty when the property that changed came from a Task or task instance object */
 	void PostEditChangeTask(FPropertyChangedChainEvent& InPropertyChangedEvent, UE::SceneState::ETaskObjectType InChangedObject);
 
 	/** Create context for this task node */
@@ -97,7 +97,7 @@ private:
 	/** Rebuilds the task instance by recreating any instanced objects / or re-initializing any 'instance' guid property */
 	void RebuildTaskInstance();
 
-	/** Called when the task type has changed */
+	/** Called when the task has changed its type, or a property related to the task / task instance has changed. */
 	void OnTaskChanged();
 
 	/** Called when a struct id has changed */
