Motion Design: added comments to some scene state files

diff --git a/Engine/Plugins/Experimental/SceneState/Source/SceneState/Public/Conduit/SceneStateConduitInstance.h b/Engine/Plugins/Experimental/SceneState/Source/SceneState/Public/Conduit/SceneStateConduitInstance.h
--- a/Engine/Plugins/Experimental/SceneState/Source/SceneState/Public/Conduit/SceneStateConduitInstance.h
+++ b/Engine/Plugins/Experimental/SceneState/Source/SceneState/Public/Conduit/SceneStateConduitInstance.h
@@ -11,6 +11,10 @@ struct FSceneStateConduitInstance
 {
 	GENERATED_BODY()
 
+	/**
+	 * Flag to determine if the conduit instance has already been initialized.
+	 * Used as a way to avoid re-setting up all the exit transitions of a conduit.
+	 */
 	UPROPERTY()
 	bool bInitialized = false;
 };
diff --git a/Engine/Plugins/Experimental/SceneState/Source/SceneState/Public/SceneState.h b/Engine/Plugins/Experimental/SceneState/Source/SceneState/Public/SceneState.h
--- a/Engine/Plugins/Experimental/SceneState/Source/SceneState/Public/SceneState.h
+++ b/Engine/Plugins/Experimental/SceneState/Source/SceneState/Public/SceneState.h
@@ -32,7 +32,7 @@ struct FSceneStateMetadata
 };
 
 /**
- * Runtime information of a given State.
+ * Runtime-immutable information of a state.
  * Holds the range of tasks, state machines and transitions belonging to the state.
  * These states are stored in the Scene State Generated Class.
  */
@@ -41,10 +41,13 @@ struct FSceneState
 {
 	GENERATED_BODY()
 
+	/** Called the state becomes the active state in its parent state machine */
 	SCENESTATE_API void Enter(const FSceneStateExecutionContext& InContext) const;
 
+	/** Called on every tick while the state is active */
 	SCENESTATE_API void Tick(const FSceneStateExecutionContext& InContext, float InDeltaSeconds) const;
 
+	/** Called when the state is no longer active in its parent state machine */
 	SCENESTATE_API void Exit(const FSceneStateExecutionContext& InContext) const;
 
 	FSceneStateRange GetTaskRange() const
diff --git a/Engine/Plugins/Experimental/SceneState/Source/SceneState/Public/SceneStateExecutionContext.h b/Engine/Plugins/Experimental/SceneState/Source/SceneState/Public/SceneStateExecutionContext.h
--- a/Engine/Plugins/Experimental/SceneState/Source/SceneState/Public/SceneStateExecutionContext.h
+++ b/Engine/Plugins/Experimental/SceneState/Source/SceneState/Public/SceneStateExecutionContext.h
@@ -71,6 +71,7 @@ struct FSceneStateExecutionContext
 	/** Gets the registry that this context is registered to */
 	TWeakPtr<const UE::SceneState::FExecutionContextRegistry> GetContextRegistry() const;
 
+	/** Gets the binding collection owned by the generated class */
 	const FSceneStateBindingCollection* GetBindingCollection() const;
 
 	/** Returns the Player Debug Name for logging purposes */
diff --git a/Engine/Plugins/Experimental/SceneState/Source/SceneState/Public/SceneStateGeneratedClass.h b/Engine/Plugins/Experimental/SceneState/Source/SceneState/Public/SceneStateGeneratedClass.h
--- a/Engine/Plugins/Experimental/SceneState/Source/SceneState/Public/SceneStateGeneratedClass.h
+++ b/Engine/Plugins/Experimental/SceneState/Source/SceneState/Public/SceneStateGeneratedClass.h
@@ -44,6 +44,7 @@ class USceneStateGeneratedClass : public UBlueprintGeneratedClass
 public:
 	SCENESTATE_API USceneStateGeneratedClass();
 
+	/** Gets the state that owns all the "top-level" state machines */
 	const FSceneState* GetRootState() const;
 
 	/** Finds the Top-Level State Machine that is linked to the given Guid */
@@ -89,6 +90,7 @@ public:
 		return Transitions;
 	}
 
+	/** Gets a copy of the transition parameters at the given transition absolute index. */
 	FInstancedPropertyBag GetTransitionParameters(uint16 InAbsoluteIndex) const;
 
 	TConstArrayView<FSceneStateMachine> GetStateMachines() const
@@ -169,9 +171,11 @@ private:
 	UPROPERTY()
 	uint16 RootStateIndex = FSceneStateRange::InvalidIndex;
 
+	/** Holds all the compiled states including nested states */
 	UPROPERTY()
 	TArray<FSceneState> States;
 
+	/** Holds all the compiled tasks (including tasks from nested states) in a contiguous block of memory */
 	UPROPERTY()
 	FInstancedStructContainer Tasks;
 
@@ -183,9 +187,11 @@ private:
 	UPROPERTY()
 	FInstancedStructContainer TaskInstances;
 
+	/** Holds all the compiled event handlers (including from nested sources) contiguously */
 	UPROPERTY()
 	TArray<FSceneStateEventHandler> EventHandlers;
 
+	/** Holds all the exit transitions (including nested sources) contiguously */
 	UPROPERTY()
 	TArray<FSceneStateTransition> Transitions;
 
@@ -213,12 +219,15 @@ private:
 	TMap<FGuid, uint16> StateMachineIdToIndex;
 
 #if WITH_EDITORONLY_DATA
+	/** Holds all the state editor-only metadata. This array matches in index with the states array. */
 	UPROPERTY()
 	TArray<FSceneStateMetadata> StateMetadata;
 
+	/** Holds all the task editor-only metadata. This array matches in index with the tasks array. */
 	UPROPERTY()
 	TArray<FSceneStateTaskMetadata> TaskMetadata;
 
+	/** Holds all the transition editor-only metadata. This array matches in index with the transitions array. */
 	UPROPERTY()
 	TArray<FSceneStateTransitionMetadata> TransitionMetadata;
 
diff --git a/Engine/Plugins/Experimental/SceneState/Source/SceneState/Public/SceneStateMachine.h b/Engine/Plugins/Experimental/SceneState/Source/SceneState/Public/SceneStateMachine.h
--- a/Engine/Plugins/Experimental/SceneState/Source/SceneState/Public/SceneStateMachine.h
+++ b/Engine/Plugins/Experimental/SceneState/Source/SceneState/Public/SceneStateMachine.h
@@ -24,6 +24,11 @@ namespace UE::SceneState
 	}
 }
 
+/**
+ * Runtime-immutable information of a state machine.
+ * Holds the range of states it will run, as well as the conduits found in the state machine, and other template data like template state machine parameters.
+ * These state machines are stored in the Scene State Generated Class.
+ */
 USTRUCT()
 struct FSceneStateMachine
 {
diff --git a/Engine/Plugins/Experimental/SceneState/Source/SceneState/Public/SceneStateObject.h b/Engine/Plugins/Experimental/SceneState/Source/SceneState/Public/SceneStateObject.h
--- a/Engine/Plugins/Experimental/SceneState/Source/SceneState/Public/SceneStateObject.h
+++ b/Engine/Plugins/Experimental/SceneState/Source/SceneState/Public/SceneStateObject.h
@@ -23,6 +23,11 @@ namespace UE::SceneState
 	class FExecutionContextRegistry;
 }
 
+/**
+ * Object instanced by USceneStateGeneratedClass.
+ * Holds the execution context (and instance data) that will be used to run the Scene State the USceneStateGeneratedClass holds.
+ * @see USceneStateGeneratedClass
+ */
 UCLASS(MinimalAPI, Blueprintable)
 class USceneStateObject : public UObject
 {
@@ -33,8 +38,10 @@ class USceneStateObject : public UObject
 public:
 	SCENESTATE_API USceneStateObject();
 
+	/** Gets the context name of the scene state player */
 	SCENESTATE_API FString GetContextName() const;
 
+	/** Gets the context object of the scene state player */
 	UFUNCTION(BlueprintCallable, Category="State")
 	SCENESTATE_API UObject* GetContextObject() const;
 
@@ -48,12 +55,16 @@ public:
 	UFUNCTION(BlueprintPure, Category="State")
 	SCENESTATE_API bool IsActive() const;
 
+	/** Called by the scene state player to setup this object and the root execution context */
 	void Setup();
 
+	/** Called by the scene state player to start execution of the root state  */
 	void Enter();
 
+	/** Called by the scene state player to start ticking the root state */
 	void Tick(float InDeltaSeconds);
 
+	/** Called by the scene state player to stop the execution of the root state */
 	void Exit();
 
 	UFUNCTION(BlueprintImplementableEvent, DisplayName = "Enter", Category = "State Events")
@@ -73,6 +84,7 @@ public:
 	//~ End UObject
 
 private:
+	/** The class of this scene state object */
 	UPROPERTY(Transient, DuplicateTransient, TextExportTransient, NonTransactional, meta=(NoBinding))
 	TObjectPtr<USceneStateGeneratedClass> GeneratedClass;
 
diff --git a/Engine/Plugins/Experimental/SceneState/Source/SceneState/Public/SceneStatePlayer.h b/Engine/Plugins/Experimental/SceneState/Source/SceneState/Public/SceneStatePlayer.h
--- a/Engine/Plugins/Experimental/SceneState/Source/SceneState/Public/SceneStatePlayer.h
+++ b/Engine/Plugins/Experimental/SceneState/Source/SceneState/Public/SceneStatePlayer.h
@@ -24,6 +24,7 @@ public:
 		return SceneStateClass;
 	}
 
+	/** Sets the class of the root state object. Destroying the existing one */
 	SCENESTATE_API void SetSceneStateClass(TSubclassOf<USceneStateObject> InSceneStateClass);
 
 	USceneStateObject* GetSceneState() const
@@ -37,12 +38,16 @@ public:
 	/** Returns the context object for this player */
 	SCENESTATE_API UObject* GetContextObject() const;
 
+	/** Should be called prior to the start of the execution */
 	SCENESTATE_API void Setup();
 
+	/** Should be called to start the execution of the scene state object */
 	SCENESTATE_API void Begin();
 
+	/** Should be called every tick during the execution of the scene state object */
 	SCENESTATE_API void Tick(float InDeltaTime);
 
+	/** Should be called to end the execution of the scene state object */
 	SCENESTATE_API void End();
 
 	//~ Begin UObject
diff --git a/Engine/Plugins/Experimental/SceneState/Source/SceneState/Public/Tasks/SceneStateBlueprintableTask.h b/Engine/Plugins/Experimental/SceneState/Source/SceneState/Public/Tasks/SceneStateBlueprintableTask.h
--- a/Engine/Plugins/Experimental/SceneState/Source/SceneState/Public/Tasks/SceneStateBlueprintableTask.h
+++ b/Engine/Plugins/Experimental/SceneState/Source/SceneState/Public/Tasks/SceneStateBlueprintableTask.h
@@ -15,6 +15,12 @@ class USceneStateObject;
 struct FSceneStateBlueprintableTaskWrapper;
 struct FSceneStateExecutionContext;
 
+/**
+ * Base class for Blueprint Tasks.
+ * Following the design of blueprints, unlike C++ Tasks that have Logic and Instance Data separated, 
+ * Blueprint Tasks holds the logic and mutable instance data together and so get allocated when the task starts.
+ * @see FSceneStateBlueprintableTaskWrapper
+ */
 UCLASS(MinimalAPI, Abstract, Blueprintable, BlueprintType, AutoExpandCategories=("Task Settings"), EditInlineNew)
 class USceneStateBlueprintableTask : public UObject, public ISceneStateEventHandlerProvider
 {
@@ -39,17 +45,22 @@ public:
 	UFUNCTION(BlueprintImplementableEvent, Category="Task")
 	void ReceiveStop(ESceneStateTaskStopReason InStopReason);
 
+	/** Gets the scene state object that owns this blueprint task */
 	UFUNCTION(BlueprintCallable, Category="Task")
 	USceneStateObject* GetRootState() const;
 
+	/** Gets the execution context of the task */
 	const FSceneStateExecutionContext& GetExecutionContext() const;
 
+	/** Gets the context object of the scene state object */
 	UFUNCTION(BlueprintCallable, Category="Task")
 	UObject* GetContextObject() const;
 
+	/** Gets the scene state object's event stream */
 	UFUNCTION(BlueprintCallable, Category="Task")
 	USceneStateEventStream* GetEventStream() const;
 
+	/** Called to indicate the task has completed */
 	UFUNCTION(BlueprintCallable, Category="Task")
 	void FinishTask();
 
diff --git a/Engine/Plugins/Experimental/SceneState/Source/SceneState/Public/Tasks/SceneStateBlueprintableTaskWrapper.h b/Engine/Plugins/Experimental/SceneState/Source/SceneState/Public/Tasks/SceneStateBlueprintableTaskWrapper.h
--- a/Engine/Plugins/Experimental/SceneState/Source/SceneState/Public/Tasks/SceneStateBlueprintableTaskWrapper.h
+++ b/Engine/Plugins/Experimental/SceneState/Source/SceneState/Public/Tasks/SceneStateBlueprintableTaskWrapper.h
@@ -14,10 +14,17 @@ struct FSceneStateBlueprintableTaskInstance : public FSceneStateTaskInstance
 {
 	GENERATED_BODY()
 
+	/** Task Blueprint Instance allocated everytime the Blueprintable Task Wrapper task starts */
 	UPROPERTY(VisibleAnywhere, Instanced, Category="Scene State")
 	TObjectPtr<USceneStateBlueprintableTask> Task;
 };
 
+/**
+ * Task Wrapper for Blueprint Tasks.
+ * Forwards the task start, tick and stop to the Blueprint Task held in its Task Instance struct
+ * @see FSceneStateBlueprintableTaskInstance
+ * @see USceneStateBlueprintableTask
+ */
 USTRUCT(meta=(Hidden))
 struct FSceneStateBlueprintableTaskWrapper : public FSceneStateTask
 {
diff --git a/Engine/Plugins/Experimental/SceneState/Source/SceneState/Public/Tasks/SceneStateMachineTask.h b/Engine/Plugins/Experimental/SceneState/Source/SceneState/Public/Tasks/SceneStateMachineTask.h
--- a/Engine/Plugins/Experimental/SceneState/Source/SceneState/Public/Tasks/SceneStateMachineTask.h
+++ b/Engine/Plugins/Experimental/SceneState/Source/SceneState/Public/Tasks/SceneStateMachineTask.h
@@ -35,7 +35,8 @@ struct FSceneStateMachineTaskInstance : public FSceneStateTaskInstance
 	FInstancedPropertyBag Parameters;
 };
 
-USTRUCT(DisplayName="Run State Machine", Category="Core", meta=(ToolTip="Runs a local state machine"))
+/** Task that runs a local state machine */
+USTRUCT(DisplayName="Run State Machine", Category="Core")
 struct FSceneStateMachineTask : public FSceneStateTask
 {
 	GENERATED_BODY()
diff --git a/Engine/Plugins/Experimental/SceneState/Source/SceneState/Public/Tasks/SceneStateMachineTaskBinding.h b/Engine/Plugins/Experimental/SceneState/Source/SceneState/Public/Tasks/SceneStateMachineTaskBinding.h
--- a/Engine/Plugins/Experimental/SceneState/Source/SceneState/Public/Tasks/SceneStateMachineTaskBinding.h
+++ b/Engine/Plugins/Experimental/SceneState/Source/SceneState/Public/Tasks/SceneStateMachineTaskBinding.h
@@ -7,6 +7,10 @@
 
 struct FSceneStateMachineTaskInstance;
 
+/**
+ * Binding Extension for the State Machine Task.
+ * Allows the State Machine Parameters to be discoverable as a binding desc
+ */
 USTRUCT()
 struct FSceneStateMachineTaskBinding : public FSceneStateTaskBindingExtension
 {
@@ -29,6 +33,7 @@ protected:
 	//~ End FSceneStateTaskBindingExtension
 
 private:
+	/** Index to the binding batch for the parameters property bag */
 	UPROPERTY()
 	uint16 ParametersBatchIndex = std::numeric_limits<uint16>::max();
 };
diff --git a/Engine/Plugins/Experimental/SceneState/Source/SceneState/Public/Tasks/SceneStateTask.h b/Engine/Plugins/Experimental/SceneState/Source/SceneState/Public/Tasks/SceneStateTask.h
--- a/Engine/Plugins/Experimental/SceneState/Source/SceneState/Public/Tasks/SceneStateTask.h
+++ b/Engine/Plugins/Experimental/SceneState/Source/SceneState/Public/Tasks/SceneStateTask.h
@@ -25,6 +25,12 @@ namespace UE::SceneState
 	}
 }
 
+/**
+ * Base class for Tasks.
+ * Tasks are immutable in execution time, and so are meant to only hold logic and template read-only data.
+ * Each Task has a Task Instance Type that is allocated and used for instance data that the task can then mutate.
+ * @see FSceneStateTaskInstance
+ */
 USTRUCT(meta=(Hidden))
 struct FSceneStateTask
 {
@@ -60,14 +66,31 @@ struct FSceneStateTask
 	/** Finds the Task Instance for the given Execution Instance */
 	SCENESTATE_API FStructView FindTaskInstance(const FSceneStateExecutionContext& InContext) const;
 
+	/**
+	 * Called when the State holding the Task first starts.
+	 * This is called for all tasks held by the state even if these tasks end up not running, or run at a later time.
+	 */
 	SCENESTATE_API void Setup(const FSceneStateExecutionContext& InContext, FStructView InTaskInstance) const;
 
+	/**
+	 * Called once the State processes that the Task has all its prerequisites met.
+	 * Applies and property bindings to the task and calls OnStart().
+	 */
 	SCENESTATE_API void Start(const FSceneStateExecutionContext& InContext, FStructView InTaskInstance) const;
 
+	/**
+	 * Called by the State each Tick.
+	 * OnTick() gets called if the Task is running and set to Tickable (see TaskFlags)
+	 */
 	SCENESTATE_API void Tick(const FSceneStateExecutionContext& InContext, FStructView InTaskInstance, float InDeltaSeconds) const;
 
+	/**
+	 * Called to Stop the Task, because it either finished or was forcibly stopped by the state
+	 * @see ESceneStateTaskStopReason
+	 */
 	SCENESTATE_API void Stop(const FSceneStateExecutionContext& InContext, FStructView InTaskInstance, ESceneStateTaskStopReason InStopReason) const;
 
+	/** Calls Stop with stop reason being that the task has finished */
 	SCENESTATE_API void Finish(const FSceneStateExecutionContext& InContext, FStructView InTaskInstance) const;
 
 	uint16 GetParentStateIndex() const
@@ -127,18 +150,28 @@ protected:
 		return nullptr;
 	}
 
+	/**
+	 * Called when the State holding the Task first starts.
+	 * This is called for all tasks held by the state even if these tasks end up not running, or run at a later time.
+	 */
 	virtual void OnSetup(const FSceneStateExecutionContext& InContext, FStructView InTaskInstance) const
 	{
 	}
 
+	/** Called once the State processes that the Task has all its prerequisites met. */
 	virtual void OnStart(const FSceneStateExecutionContext& InContext, FStructView InTaskInstance) const
 	{
 	}
 
+	/** Called by the State each Tick if the Task is set to Tick via its TaskFlags */
 	virtual void OnTick(const FSceneStateExecutionContext& InContext, FStructView InTaskInstance, float InDeltaSeconds) const
 	{
 	}
 
+	/**
+	 * Called when the Task has finished or forcibly stopped by the state
+	 * @see ESceneStateTaskStopReason
+	 */
 	virtual void OnStop(const FSceneStateExecutionContext& InContext, FStructView InTaskInstance, ESceneStateTaskStopReason InStopReason) const
 	{
 	}
diff --git a/Engine/Plugins/Experimental/SceneState/Source/SceneState/Public/Tasks/SceneStateTaskMetadata.h b/Engine/Plugins/Experimental/SceneState/Source/SceneState/Public/Tasks/SceneStateTaskMetadata.h
--- a/Engine/Plugins/Experimental/SceneState/Source/SceneState/Public/Tasks/SceneStateTaskMetadata.h
+++ b/Engine/Plugins/Experimental/SceneState/Source/SceneState/Public/Tasks/SceneStateTaskMetadata.h
@@ -11,6 +11,7 @@ struct FSceneStateTaskMetadata
 	GENERATED_BODY()
 
 #if WITH_EDITORONLY_DATA
+	/** Id used for the Property Binding System to identify a Task */
 	UPROPERTY()
 	FGuid TaskId;
 #endif
diff --git a/Engine/Plugins/Experimental/SceneState/Source/SceneState/Public/Transition/SceneStateTransition.h b/Engine/Plugins/Experimental/SceneState/Source/SceneState/Public/Transition/SceneStateTransition.h
--- a/Engine/Plugins/Experimental/SceneState/Source/SceneState/Public/Transition/SceneStateTransition.h
+++ b/Engine/Plugins/Experimental/SceneState/Source/SceneState/Public/Transition/SceneStateTransition.h
@@ -33,6 +33,10 @@ namespace UE::SceneState
 	}
 }
 
+/**
+ * Holds the exit transition conditions and information from its owner state to a target.
+ * @see FSceneStateTransitionTarget
+ */
 USTRUCT()
 struct FSceneStateTransition
 {
diff --git a/Engine/Plugins/Experimental/SceneState/Source/SceneState/Public/Transition/SceneStateTransitionResult.h b/Engine/Plugins/Experimental/SceneState/Source/SceneState/Public/Transition/SceneStateTransitionResult.h
--- a/Engine/Plugins/Experimental/SceneState/Source/SceneState/Public/Transition/SceneStateTransitionResult.h
+++ b/Engine/Plugins/Experimental/SceneState/Source/SceneState/Public/Transition/SceneStateTransitionResult.h
@@ -4,11 +4,13 @@
 
 #include "SceneStateTransitionResult.generated.h"
 
+/** The result output of a transition graph evaluation */
 USTRUCT(BlueprintInternalUseOnly)
 struct FSceneStateTransitionResult
 {
 	GENERATED_BODY()
 
+	/** Indicator of whether the transition condition succeeded */
 	UPROPERTY(EditAnywhere, BlueprintReadWrite, Category="Scene State Transition")
 	bool bCanTransition = true;
 };
diff --git a/Engine/Plugins/Experimental/SceneState/Source/SceneState/Public/Transition/SceneStateTransitionTarget.h b/Engine/Plugins/Experimental/SceneState/Source/SceneState/Public/Transition/SceneStateTransitionTarget.h
--- a/Engine/Plugins/Experimental/SceneState/Source/SceneState/Public/Transition/SceneStateTransitionTarget.h
+++ b/Engine/Plugins/Experimental/SceneState/Source/SceneState/Public/Transition/SceneStateTransitionTarget.h
@@ -5,6 +5,7 @@
 #include <limits>
 #include "SceneStateTransitionTarget.generated.h"
 
+/** Types of object a transition can lead to */
 UENUM()
 enum class ESceneStateTransitionTargetType : uint8
 {
diff --git a/Engine/Plugins/Experimental/SceneState/Source/SceneStateBinding/Public/SceneStateBinding.h b/Engine/Plugins/Experimental/SceneState/Source/SceneStateBinding/Public/SceneStateBinding.h
--- a/Engine/Plugins/Experimental/SceneState/Source/SceneStateBinding/Public/SceneStateBinding.h
+++ b/Engine/Plugins/Experimental/SceneState/Source/SceneStateBinding/Public/SceneStateBinding.h
@@ -6,6 +6,7 @@
 #include "SceneStateBindingDataHandle.h"
 #include "SceneStateBinding.generated.h"
 
+/** Represents a binding from a source property to a destination property */
 USTRUCT()
 struct FSceneStateBinding : public FPropertyBindingBinding
 {
@@ -17,9 +18,11 @@ struct FSceneStateBinding : public FPropertyBindingBinding
 	SCENESTATEBINDING_API virtual FConstStructView GetSourceDataHandleStruct() const override;
 	//~ End FPropertyBindingBinding
 
+	/** Handle to the data containing the source property */
 	UPROPERTY()
 	FSceneStateBindingDataHandle SourceDataHandle;
 
+	/** Handle to the data containing the target property */
 	UPROPERTY()
 	FSceneStateBindingDataHandle TargetDataHandle;
 };
diff --git a/Engine/Plugins/Experimental/SceneState/Source/SceneStateBinding/Public/SceneStateBindingCollection.h b/Engine/Plugins/Experimental/SceneState/Source/SceneStateBinding/Public/SceneStateBindingCollection.h
--- a/Engine/Plugins/Experimental/SceneState/Source/SceneStateBinding/Public/SceneStateBindingCollection.h
+++ b/Engine/Plugins/Experimental/SceneState/Source/SceneStateBinding/Public/SceneStateBindingCollection.h
@@ -15,6 +15,7 @@ namespace UE::SceneState::Editor
 	class FBindingCompiler;
 }
 
+/** Holds all the property bindings to use in execution */
 USTRUCT()
 struct FSceneStateBindingCollection : public FPropertyBindingBindingCollection
 {
