Motion Design: Moved scene state and data link plugins out of experimental into virtualproduction

diff --git a/Engine/Plugins/VirtualProduction/SceneState/Source/SceneState/Private/SceneState.cpp b/Engine/Plugins/VirtualProduction/SceneState/Source/SceneState/Private/SceneState.cpp
new file mode 100644
--- /dev/null
+++ b/Engine/Plugins/VirtualProduction/SceneState/Source/SceneState/Private/SceneState.cpp
@@ -0,0 +1,331 @@
+// Copyright Epic Games, Inc. All Rights Reserved.
+
+#include "SceneState.h"
+#include "SceneStateEventHandler.h"
+#include "SceneStateEventStream.h"
+#include "SceneStateExecutionContext.h"
+#include "SceneStateInstance.h"
+#include "SceneStateLog.h"
+#include "SceneStateMachine.h"
+#include "SceneStateObject.h"
+#include "SceneStateUtils.h"
+#include "StructUtils/InstancedStructContainer.h"
+#include "Tasks/SceneStateTask.h"
+#include "Tasks/SceneStateTaskInstance.h"
+#include "Transition/SceneStateTransition.h"
+
+void FSceneState::Enter(const FSceneStateExecutionContext& InContext) const
+{
+	using namespace UE::SceneState;
+
+	const FReentryGuard ReentryGuard(ReentryHandle, InContext);
+	if (ReentryGuard.IsReentry())
+	{
+		return;
+	}
+
+	// Enter State, add a State Instance if not already present
+	FSceneStateInstance* Instance = InContext.FindOrAddStateInstance(*this);
+	if (!Instance || Instance->Status == EExecutionStatus::Running)
+	{
+		return;
+	}
+
+	UE_LOG(LogSceneState, Verbose, TEXT("State (%s) receiving enter"), GetStateName(InContext).GetData());
+
+	Instance->ElapsedTime = 0.f;
+	Instance->Status = EExecutionStatus::Running;
+
+	// Apply Event Handlers before anything else starts so the Event Data becomes available to substate machines and tasks
+	CaptureEvents(InContext);
+	AllocateTaskInstances(InContext, InContext.GetTemplateTaskInstances(*this));
+
+	for (const FSceneStateTransition& Transition : InContext.GetTransitions(*this))
+	{
+		Transition.Setup(InContext);
+	}
+
+	for (const FSceneStateMachine& StateMachine : InContext.GetStateMachines(*this))
+	{
+		StateMachine.Setup(InContext);
+	}
+
+	// Setup each Task
+	InContext.ForEachTask(*this,
+		[&InContext, Instance](const FSceneStateTask& InTask, FStructView InTaskInstance)->EIterationResult
+		{
+			InTask.Setup(InContext, InTaskInstance);
+			return EIterationResult::Continue;
+		});
+
+	for (const FSceneStateMachine& StateMachine : InContext.GetStateMachines(*this))
+	{
+		StateMachine.Start(InContext);
+	}
+
+	UpdateActiveTasks(InContext, *Instance);
+}
+
+void FSceneState::Tick(const FSceneStateExecutionContext& InContext, float InDeltaSeconds) const
+{
+	using namespace UE::SceneState;
+
+	const FReentryGuard ReentryGuard(ReentryHandle, InContext);
+	if (ReentryGuard.IsReentry())
+	{
+		return;
+	}
+
+	FSceneStateInstance* Instance = InContext.FindStateInstance(*this);
+	if (!Instance || Instance->Status != EExecutionStatus::Running)
+	{
+		return;
+	}
+
+	UE_LOG(LogSceneState, VeryVerbose, TEXT("State (%s) receiving tick"), GetStateName(InContext).GetData());
+
+	Instance->ElapsedTime += InDeltaSeconds;
+
+	UpdateActiveTasks(InContext, *Instance);
+
+	for (const FSceneStateMachine& StateMachine : InContext.GetStateMachines(*this))
+	{
+		StateMachine.Tick(InContext, InDeltaSeconds);
+	}
+
+	InContext.ForEachTask(*this,
+		[&InContext, Instance, InDeltaSeconds](const FSceneStateTask& InTask, FStructView InTaskInstance)
+		{
+			InTask.Tick(InContext, InTaskInstance, InDeltaSeconds);
+			return EIterationResult::Continue;
+		});
+}
+
+void FSceneState::Exit(const FSceneStateExecutionContext& InContext) const
+{
+	using namespace UE::SceneState;
+
+	const FReentryGuard ReentryGuard(ReentryHandle, InContext);
+	if (ReentryGuard.IsReentry())
+	{
+		return;
+	}
+
+	FSceneStateInstance* Instance = InContext.FindStateInstance(*this);
+	if (!Instance || Instance->Status != EExecutionStatus::Running)
+	{
+		return;
+	}
+
+	UE_LOG(LogSceneState, Verbose, TEXT("State (%s) receiving exit"), GetStateName(InContext).GetData());
+
+	// Stop State Machines that are still running
+	for (const FSceneStateMachine& StateMachine : InContext.GetStateMachines(*this))
+	{
+		StateMachine.Stop(InContext);
+	}
+
+	// Stop tasks that are still running
+	InContext.ForEachTask(*this,
+		[&InContext, Instance](const FSceneStateTask& InTask, FStructView InTaskInstance)
+		{
+			InTask.Stop(InContext, InTaskInstance, ESceneStateTaskStopReason::State);
+			return EIterationResult::Continue;
+		});
+
+	// Notify transitions of stop
+	for (const FSceneStateTransition& Transition : InContext.GetTransitions(*this))
+	{
+		Transition.Exit(InContext);
+	}
+
+	Instance->Status = EExecutionStatus::Finished;
+	Instance->ElapsedTime = 0.f;
+
+	InContext.RemoveStateInstance(*this);
+	InContext.RemoveTaskInstanceContainer(*this);
+
+	ResetCapturedEvents(InContext);
+}
+
+void FSceneState::UpdateActiveTasks(const FSceneStateExecutionContext& InContext, FSceneStateInstance& InInstance) const
+{
+	using namespace UE::SceneState;
+
+	FInstancedStructContainer* TaskInstanceContainer = InContext.FindTaskInstanceContainer(*this);
+	if (!TaskInstanceContainer)
+	{
+		return;
+	}
+
+	auto FindTaskInstance = 
+		[TaskInstanceContainer](uint16 InRelativeIndex)->FSceneStateTaskInstance*
+		{
+			if (TaskInstanceContainer->IsValidIndex(InRelativeIndex))
+			{
+				return (*TaskInstanceContainer)[InRelativeIndex].GetPtr<FSceneStateTaskInstance>();
+			}
+			return nullptr;
+		};
+
+	// Start all the Tasks that haven't started yet and that meet their pre-requisites
+	InContext.ForEachTask(*this,
+		[&InContext, &InInstance, &FindTaskInstance](const FSceneStateTask& InTask, FStructView InTaskInstance)->EIterationResult
+		{
+			FSceneStateTaskInstance* TaskInstance = InTaskInstance.GetPtr<FSceneStateTaskInstance>();
+			if (!TaskInstance || TaskInstance->GetStatus() != EExecutionStatus::NotStarted)
+			{
+				return EIterationResult::Continue;
+			}
+
+			bool bPrerequisitesMet = true;
+			for (uint16 Prerequisite : InContext.GetTaskPrerequisites(InTask))
+			{
+				FSceneStateTaskInstance* const PrerequisiteTaskInstance = FindTaskInstance(Prerequisite);
+				if (!PrerequisiteTaskInstance || PrerequisiteTaskInstance->GetStatus() != EExecutionStatus::Finished)
+				{
+					bPrerequisitesMet = false;
+					break;
+				}
+			}
+
+			if (bPrerequisitesMet)
+			{
+				InTask.Start(InContext, InTaskInstance);
+			}
+
+			return EIterationResult::Continue;
+		});
+}
+
+FStringView FSceneState::GetStateName(const FSceneStateExecutionContext& InContext) const
+{
+#if WITH_EDITOR
+	if (const FSceneStateMetadata* StateMetadata = InContext.GetStateMetadata(*this))
+	{
+		return StateMetadata->StateName;
+	}
+#endif
+	return FStringView();
+}
+
+bool FSceneState::HasPendingTasks(const FSceneStateExecutionContext& InContext) const
+{
+	using namespace UE::SceneState;
+
+	FSceneStateInstance* const StateInstance = InContext.FindStateInstance(*this);
+	if (!StateInstance)
+	{
+		return false;
+	}
+
+	bool bHasPendingTask = false;
+
+	InContext.ForEachTask(*this,
+		[StateInstance, &bHasPendingTask](const FSceneStateTask& InTask, FStructView InTaskInstance)->EIterationResult
+		{
+			FSceneStateTaskInstance* TaskInstance = InTaskInstance.GetPtr<FSceneStateTaskInstance>();
+			if (TaskInstance && TaskInstance->GetStatus() != EExecutionStatus::Finished)
+			{
+				bHasPendingTask = true;
+				return EIterationResult::Break;
+			}
+			return EIterationResult::Continue;
+		});
+
+	return bHasPendingTask;
+}
+
+void FSceneState::AllocateTaskInstances(const FSceneStateExecutionContext& InContext, TConstArrayView<FConstStructView> InTemplateTaskInstances) const
+{
+	if (FInstancedStructContainer* const TaskInstanceContainer = InContext.FindOrAddTaskInstanceContainer(*this))
+	{
+		// Copy the Template data
+		*TaskInstanceContainer = InTemplateTaskInstances;
+
+		// Instance each Template Object into the Instance data
+		InstanceTaskObjects(InContext.GetRootState()
+			, UE::SceneState::GetStructViews(*TaskInstanceContainer)
+			, InTemplateTaskInstances
+			, [](FObjectDuplicationParameters& InParams)->UObject*
+			{
+				return StaticDuplicateObjectEx(InParams);
+			});
+	}
+}
+
+void FSceneState::InstanceTaskObjects(UObject* InOuter, TConstArrayView<FStructView> InTargets, TConstArrayView<FConstStructView> InSources, TFunctionRef<UObject*(FObjectDuplicationParameters&)> InDuplicationFunc) const
+{
+	check(InTargets.Num() == InSources.Num());
+
+	for (int32 TaskIndex = 0; TaskIndex < InSources.Num(); ++TaskIndex)
+	{
+		const FConstStructView Source = InSources[TaskIndex];
+		const FStructView Target = InTargets[TaskIndex];
+
+		check(Source.GetScriptStruct() == Target.GetScriptStruct());
+
+		if (!Source.GetScriptStruct())
+		{
+			continue;
+		}
+
+		// Visit all object properties that are instanced references, stepping into structs but not object properties themselves.
+		// Using const_cast here as there's no Visit variant taking a const ptr.
+		Source.GetScriptStruct()->Visit(const_cast<uint8*>(Source.GetMemory()),
+			[&InDuplicationFunc, InOuter, Target](const FPropertyVisitorContext& InContext)->EPropertyVisitorControlFlow
+			{
+				const FObjectProperty* ObjectProperty = CastField<FObjectProperty>(InContext.Path.Top().Property);
+				if (!ObjectProperty)
+				{
+					return EPropertyVisitorControlFlow::StepInto;
+				}
+
+				if (!ObjectProperty->HasAnyPropertyFlags(CPF_InstancedReference))
+				{
+					return EPropertyVisitorControlFlow::StepOver;
+				}
+
+				const void* const SourcePropertyMemory = InContext.Data.PropertyData;
+
+				UObject* InstanceObject = nullptr;
+				if (UObject* TemplateObject = ObjectProperty->GetObjectPropertyValue(SourcePropertyMemory))
+				{
+					FObjectDuplicationParameters Parameters = InitStaticDuplicateObjectParams(TemplateObject, InOuter);
+					Parameters.DestName = MakeUniqueObjectName(Parameters.DestOuter, Parameters.SourceObject->GetClass(), Parameters.SourceObject->GetFName());
+					Parameters.FlagMask = RF_AllFlags & ~RF_DefaultSubObject;
+					Parameters.PortFlags |= PPF_DuplicateVerbatim; // Skip resetting text IDs
+
+					InstanceObject = InDuplicationFunc(Parameters);
+
+					UE_LOG(LogSceneState, Verbose, TEXT("Instantiated object '%s' (Outer: '%s') from template '%s' (Outer: '%s')")
+						, *GetNameSafe(InstanceObject)
+						, *InOuter->GetName()
+						, *TemplateObject->GetName()
+						, *GetNameSafe(TemplateObject->GetOuter()));
+				}
+
+				void* const TargetPropertyMemory = UE::SceneState::ResolveVisitedPath(Target.GetScriptStruct(), Target.GetMemory(), InContext.Path);
+				ObjectProperty->SetPropertyValue(TargetPropertyMemory, InstanceObject);
+				return EPropertyVisitorControlFlow::StepOver;
+			});
+	}
+}
+
+void FSceneState::CaptureEvents(const FSceneStateExecutionContext& InContext) const
+{
+	if (USceneStateEventStream* EventStream = InContext.GetEventStream())
+	{
+		TConstArrayView<FSceneStateEventHandler> EventHandlers = InContext.GetEventHandlers(*this);
+		EventStream->CaptureEvents(EventHandlers);
+	}
+}
+
+void FSceneState::ResetCapturedEvents(const FSceneStateExecutionContext& InContext) const
+{
+	if (USceneStateEventStream* EventStream = InContext.GetEventStream())
+	{
+		TConstArrayView<FSceneStateEventHandler> EventHandlers = InContext.GetEventHandlers(*this);
+		EventStream->ResetCapturedEvents(EventHandlers);
+	}
+}
diff --git a/Engine/Plugins/VirtualProduction/SceneState/Source/SceneState/Private/SceneStateExecutionContext.cpp b/Engine/Plugins/VirtualProduction/SceneState/Source/SceneState/Private/SceneStateExecutionContext.cpp
new file mode 100644
--- /dev/null
+++ b/Engine/Plugins/VirtualProduction/SceneState/Source/SceneState/Private/SceneStateExecutionContext.cpp
@@ -0,0 +1,769 @@
+// Copyright Epic Games, Inc. All Rights Reserved.
+
+#include "SceneStateExecutionContext.h"
+#include "PropertyBindingDataView.h"
+#include "PropertyBindingTypes.h"
+#include "SceneState.h"
+#include "SceneStateExecutionContextRegistry.h"
+#include "SceneStateGeneratedClass.h"
+#include "SceneStateLog.h"
+#include "SceneStateMachine.h"
+#include "SceneStateObject.h"
+#include "SceneStateUtils.h"
+#include "Tasks/SceneStateTaskBindingExtension.h"
+#include "Tasks/SceneStateTaskInstance.h"
+
+const FSceneStateExecutionContext FSceneStateExecutionContext::InvalidContext;
+
+FSceneStateExecutionContext::~FSceneStateExecutionContext()
+{
+	Reset();
+}
+
+TWeakPtr<const UE::SceneState::FExecutionContextRegistry> FSceneStateExecutionContext::GetContextRegistry() const
+{
+	return ContextRegistryWeak;
+}
+
+const FSceneStateBindingCollection* FSceneStateExecutionContext::GetBindingCollection() const
+{
+	if (GeneratedClass)
+	{
+		return &GeneratedClass->GetBindingCollection();
+	}
+	return nullptr;
+}
+
+FString FSceneStateExecutionContext::GetExecutionContextName() const
+{
+	if (RootState)
+	{
+		return RootState->GetContextName();
+	}
+	return FString();
+}
+
+UObject* FSceneStateExecutionContext::GetContextObject() const
+{
+	if (RootState)
+	{
+		return RootState->GetContextObject();
+	}
+	return nullptr;
+}
+
+USceneStateEventStream* FSceneStateExecutionContext::GetEventStream() const
+{
+	if (RootState)
+	{
+		return RootState->GetEventStream();
+	}
+	return nullptr;
+}
+
+void FSceneStateExecutionContext::Setup(const FSceneStateExecutionContext& InParentContext)
+{
+	Setup(InParentContext.GetRootState());
+	check(RootState);
+
+	const TSharedRef<UE::SceneState::FExecutionContextRegistry> ContextRegistry = RootState->GetContextRegistry();
+	ParentHandle = ContextRegistry->FindHandle(InParentContext);
+}
+
+void FSceneStateExecutionContext::Setup(USceneStateObject* InRootState)
+{
+	Reset();
+
+	check(InRootState);
+	RootState = InRootState;
+	GeneratedClass = Cast<USceneStateGeneratedClass>(RootState->GetClass());
+
+	const TSharedRef<UE::SceneState::FExecutionContextRegistry> ContextRegistry = InRootState->GetContextRegistry();
+	ContextRegistry->RegisterContext(this);
+	ContextRegistryWeak = ContextRegistry;
+}
+
+void FSceneStateExecutionContext::Reset()
+{
+	RootState = nullptr;
+	GeneratedClass = nullptr;
+
+	StateInstances.Reset();
+	StateMachineInstances.Reset();
+	TransitionInstances.Reset();
+
+	if (TSharedPtr<UE::SceneState::FExecutionContextRegistry> ContextRegistry = ContextRegistryWeak.Pin())
+	{
+		ContextRegistry->UnregisterContext(this);
+	}
+
+	ContextRegistryWeak.Reset();
+}
+
+FPropertyBindingDataView FSceneStateExecutionContext::FindDataView(const FSceneStateBindingDataHandle& InDataHandle) const
+{
+	const uint16 DataIndex = InDataHandle.GetDataIndex();
+	const uint16 DataSubIndex = InDataHandle.GetDataSubIndex();
+
+	if (InDataHandle.IsExternalDataType())
+	{
+		// Todo: unsupported external sources
+		return FPropertyBindingDataView();
+	}
+
+	const ESceneStateDataType DataType = static_cast<ESceneStateDataType>(InDataHandle.GetDataType());
+
+	switch (DataType)
+	{
+	case ESceneStateDataType::Root:
+		return FPropertyBindingDataView(RootState);
+
+	case ESceneStateDataType::Task:
+		return FindTaskInstance(DataIndex);
+
+	case ESceneStateDataType::EventHandler:
+		if (USceneStateEventStream* EventStream = GetEventStream())
+		{
+			if (const FSceneStateEventHandler* EventHandler = GetEventHandler(DataIndex))
+			{
+				if (FSceneStateEvent* Event = EventStream->FindCapturedEvent(EventHandler->GetHandlerId()))
+				{
+					return Event->GetDataViewMutable();
+				}
+			}
+		}
+		break;
+
+	case ESceneStateDataType::TaskExtension:
+		if (const FSceneStateTask* Task = FindTask(DataIndex).GetPtr<const FSceneStateTask>())
+		{
+			if (const FSceneStateTaskBindingExtension* BindingExtension = Task->GetBindingExtension())
+			{
+				FStructView DataView;
+				if (BindingExtension->FindDataByIndex(FindTaskInstance(DataIndex), DataSubIndex, DataView))
+				{
+					return DataView;
+				}
+			}
+		}
+		break;
+
+	case ESceneStateDataType::Transition:
+		if (FSceneStateTransitionInstance* TransitionInstance = FindTransitionInstance(DataIndex))
+		{
+			return TransitionInstance->Parameters.GetMutableValue();
+		}
+		break;
+
+	case ESceneStateDataType::StateMachine:
+		if (FSceneStateMachineInstance* StateMachineInstance = FindStateMachineInstance(DataIndex))
+		{
+			return StateMachineInstance->Parameters.GetMutableValue();
+		}
+		break;
+	}
+
+	// The data handle could be exclusively on the parent context, so call FindDataView on the parent too
+	if (const FSceneStateExecutionContext* ParentContext = FindParentContext())
+	{
+		return ParentContext->FindDataView(InDataHandle);
+	}
+
+	return FPropertyBindingDataView();
+}
+
+void FSceneStateExecutionContext::ForEachTask(const FSceneState& InState, TFunctionRef<UE::SceneState::EIterationResult(const FSceneStateTask&, FStructView)> InCallable) const
+{
+	const FSceneStateRange TaskRange = InState.GetTaskRange();
+	if (TaskRange.Count == 0 || !GeneratedClass)
+	{
+		return;
+	}
+
+	FInstancedStructContainer* TaskInstanceContainer = FindTaskInstanceContainer(InState);
+	if (!TaskInstanceContainer)
+	{
+		return;
+	}
+
+	const FInstancedStructContainer& Tasks = GeneratedClass->GetTasks();
+
+	checkfSlow(Tasks.IsValidIndex(TaskRange.Index) && Tasks.IsValidIndex(TaskRange.GetLastIndex())
+		, TEXT("[Context: %s] State Task Range [%d, %d] is out of bounds! (Total Task Count: %d)")
+		, *GetExecutionContextName()
+		, TaskRange.Index
+		, TaskRange.GetLastIndex()
+		, Tasks.Num());
+
+	for (uint16 TaskIndex = 0; TaskIndex < TaskRange.Count; ++TaskIndex)
+	{
+		if (!TaskInstanceContainer->IsValidIndex(TaskIndex))
+		{
+			continue;
+		}
+
+		const uint16 AbsoluteIndex = TaskRange.Index + TaskIndex;
+		const FSceneStateTask& Task = Tasks[AbsoluteIndex].Get<const FSceneStateTask>();
+		const FStructView TaskInstance = (*TaskInstanceContainer)[TaskIndex];
+
+		if (InCallable(Task, TaskInstance) == UE::SceneState::EIterationResult::Break)
+		{
+			break;
+		}
+	}
+}
+
+TArray<FConstStructView> FSceneStateExecutionContext::GetTemplateTaskInstances(const FSceneState& InState) const
+{
+	const FSceneStateRange TaskRange = InState.GetTaskRange();
+	if (TaskRange.Count == 0 || !GeneratedClass)
+	{
+		return {};
+	}
+	return UE::SceneState::GetConstStructViews(GeneratedClass->GetTaskInstances(), TaskRange);
+}
+
+TConstArrayView<uint16> FSceneStateExecutionContext::GetTaskPrerequisites(const FSceneStateTask& InTask) const
+{
+	const FSceneStateRange TaskPrerequisiteRange = InTask.GetPrerequisiteRange();
+	if (TaskPrerequisiteRange.Count == 0 || !GeneratedClass)
+	{
+		return {};
+	}
+
+	const TConstArrayView<uint16> TaskPrerequisitesView = GeneratedClass->GetTaskPrerequisites();
+
+	checkfSlow(TaskPrerequisitesView.IsValidIndex(TaskPrerequisiteRange.Index) && TaskPrerequisitesView.IsValidIndex(TaskPrerequisiteRange.GetLastIndex())
+		, TEXT("[Context: %s] Task Prerequisite Range [%d, %d] is out of bounds! (Total Prerequisite Count: %d)")
+		, *GetExecutionContextName()
+		, TaskPrerequisiteRange.Index
+		, TaskPrerequisiteRange.GetLastIndex()
+		, TaskPrerequisitesView.Num());
+
+	return TaskPrerequisitesView.Slice(TaskPrerequisiteRange.Index, TaskPrerequisiteRange.Count);
+}
+
+TConstArrayView<FSceneStateTransition> FSceneStateExecutionContext::GetTransitions(const FSceneState& InState) const
+{
+	const FSceneStateRange TransitionRange = InState.GetTransitionRange();
+	if (TransitionRange.Count == 0 || !GeneratedClass)
+	{
+		return {};
+	}
+
+	const TConstArrayView<FSceneStateTransition> TransitionsView = GeneratedClass->GetTransitions();
+
+	checkfSlow(TransitionsView.IsValidIndex(TransitionRange.Index) && TransitionsView.IsValidIndex(TransitionRange.GetLastIndex())
+		, TEXT("[Context: %s] State Transition Range [%d, %d] is out of bounds! (Total Transition Count: %d)")
+		, *GetExecutionContextName()
+		, TransitionRange.Index
+		, TransitionRange.GetLastIndex()
+		, TransitionsView.Num());
+
+	return TransitionsView.Slice(TransitionRange.Index, TransitionRange.Count);
+}
+
+TConstArrayView<FSceneStateTransition> FSceneStateExecutionContext::GetTransitions(const FSceneStateConduit& InConduit) const
+{
+	const FSceneStateRange TransitionRange = InConduit.GetTransitionRange();
+	if (TransitionRange.Count == 0 || !GeneratedClass)
+	{
+		return {};
+	}
+
+	const TConstArrayView<FSceneStateTransition> TransitionsView = GeneratedClass->GetTransitions();
+
+	checkfSlow(TransitionsView.IsValidIndex(TransitionRange.Index) && TransitionsView.IsValidIndex(TransitionRange.GetLastIndex())
+		, TEXT("[Context: %s] Conduit Transition Range [%d, %d] is out of bounds! (Total Transition Count: %d)")
+		, *GetExecutionContextName()
+		, TransitionRange.Index
+		, TransitionRange.GetLastIndex()
+		, TransitionsView.Num());
+
+	return TransitionsView.Slice(TransitionRange.Index, TransitionRange.Count);
+}
+
+FInstancedPropertyBag FSceneStateExecutionContext::GetTemplateTransitionParameter(const FSceneStateTransition& InTransition) const
+{
+	uint16 TransitionIndex;
+	if (GetTransitionIndex(InTransition, TransitionIndex))
+	{
+		return GeneratedClass->GetTransitionParameters(TransitionIndex);
+	}
+	return FInstancedPropertyBag();
+}
+
+TConstArrayView<FSceneStateMachine> FSceneStateExecutionContext::GetStateMachines(const FSceneState& InState) const
+{
+	const FSceneStateRange StateMachineRange = InState.GetStateMachineRange();
+	if (StateMachineRange.Count == 0 || !GeneratedClass)
+	{
+		return {};
+	}
+
+	const TConstArrayView<FSceneStateMachine> StateMachinesView = GeneratedClass->GetStateMachines();
+
+	checkfSlow(StateMachinesView.IsValidIndex(StateMachineRange.Index) && StateMachinesView.IsValidIndex(StateMachineRange.GetLastIndex())
+		, TEXT("[Context: %s] State State Machine Range [%d, %d] is out of bounds! (Total State Machine Count: %d)")
+		, *GetExecutionContextName()
+		, StateMachineRange.Index
+		, StateMachineRange.GetLastIndex()
+		, StateMachinesView.Num());
+
+	return StateMachinesView.Slice(StateMachineRange.Index, StateMachineRange.Count);
+}
+
+TConstArrayView<FSceneStateEventHandler> FSceneStateExecutionContext::GetEventHandlers(const FSceneState& InState) const
+{
+	const FSceneStateRange EventHandlerRange = InState.GetEventHandlerRange();
+	if (EventHandlerRange.Count == 0 || !GeneratedClass)
+	{
+		return {};
+	}
+
+	const TConstArrayView<FSceneStateEventHandler> EventHandlersView = GeneratedClass->GetEventHandlers();
+
+	checkfSlow(EventHandlersView.IsValidIndex(EventHandlerRange.Index) && EventHandlersView.IsValidIndex(EventHandlerRange.GetLastIndex())
+		, TEXT("[Context: %s] State Event Handler Range [%d, %d] is out of bounds! (Total Event Handler Count: %d)")
+		, *GetExecutionContextName()
+		, EventHandlerRange.Index
+		, EventHandlerRange.GetLastIndex()
+		, EventHandlersView.Num());
+
+	return EventHandlersView.Slice(EventHandlerRange.Index, EventHandlerRange.Count);
+}
+
+#if WITH_EDITOR
+const FSceneStateMetadata* FSceneStateExecutionContext::GetStateMetadata(const FSceneState& InState) const
+{
+	if (!GeneratedClass)
+	{
+		return nullptr;
+	}
+
+	uint16 StateIndex;
+	if (!GetStateIndex(InState, StateIndex))
+	{
+		UE_LOG(LogSceneState, Error, TEXT("[Context: %s] State Index was not able to be retrieved. State Data is invalid!"), *GetExecutionContextName());
+		return nullptr;
+	}
+
+	const TConstArrayView<FSceneStateMetadata> StateMetadataView = GeneratedClass->GetStateMetadata();
+
+	if (!StateMetadataView.IsValidIndex(StateIndex))
+	{
+		UE_LOG(LogSceneState, Error, TEXT("[Context: %s] Metadata for State Index %d is not available!"
+			"Be sure to have compiled the Scene State Object. Metadata count: %d")
+			, *GetExecutionContextName()
+			, StateIndex
+			, StateMetadataView.Num());
+		return nullptr;
+	}
+
+	return &StateMetadataView[StateIndex];
+}
+#endif
+
+const FSceneStateMachine* FSceneStateExecutionContext::GetStateMachine(const FGuid& InStateMachineId) const
+{
+	if (GeneratedClass)
+	{
+		return GeneratedClass->FindStateMachine(InStateMachineId);
+	}
+	return nullptr;
+}
+
+const FSceneStateMachine* FSceneStateExecutionContext::GetStateMachine(uint16 InAbsoluteIndex) const
+{
+	if (!GeneratedClass)
+	{
+		return nullptr;
+	}
+
+	const TConstArrayView<FSceneStateMachine> StateMachinesView = GeneratedClass->GetStateMachines();
+
+	checkfSlow(StateMachinesView.IsValidIndex(InAbsoluteIndex)
+		, TEXT("[Context: %s] Absolute Index %d is out of bounds! (Total State Machine Count: %d)")
+		, *GetExecutionContextName()
+		, InAbsoluteIndex
+		, StateMachinesView.Num());
+
+	return &StateMachinesView[InAbsoluteIndex];
+}
+
+const FSceneStateMachine* FSceneStateExecutionContext::GetContextStateMachine() const
+{
+	if (StateMachineExecutionStack.IsEmpty() || !GeneratedClass)
+	{
+		return nullptr;
+	}
+
+	return GetStateMachine(StateMachineExecutionStack.Last());
+}
+
+const FSceneState* FSceneStateExecutionContext::GetActiveState(const FSceneStateMachine& InStateMachine) const
+{
+	if (const FSceneStateMachineInstance* Instance = FindStateMachineInstance(InStateMachine))
+	{
+		return GetState(InStateMachine, Instance->ActiveIndex);
+	}
+	return nullptr;
+}
+
+const FSceneState* FSceneStateExecutionContext::GetState(uint16 InAbsoluteIndex) const
+{
+	if (!GeneratedClass)
+	{
+		return nullptr;
+	}
+
+	const TConstArrayView<FSceneState> StatesView = GeneratedClass->GetStates();
+
+	checkfSlow(StatesView.IsValidIndex(InAbsoluteIndex)
+		, TEXT("[Context: %s] Absolute Index %d is out of bounds! (Total State Count: %d)")
+		, *GetExecutionContextName()
+		, InAbsoluteIndex
+		, StatesView.Num());
+
+	return &StatesView[InAbsoluteIndex];
+}
+
+const FSceneStateEventHandler* FSceneStateExecutionContext::GetEventHandler(uint16 InAbsoluteIndex) const
+{
+	if (!GeneratedClass)
+	{
+		return nullptr;
+	}
+
+	const TConstArrayView<FSceneStateEventHandler> EventHandlers = GeneratedClass->GetEventHandlers();
+
+	checkfSlow(EventHandlers.IsValidIndex(InAbsoluteIndex)
+		, TEXT("[Context: %s] Absolute Index %d is out of bounds! (Total Event Handlers Count: %d)")
+		, *GetExecutionContextName()
+		, InAbsoluteIndex
+		, EventHandlers.Num());
+
+	return &EventHandlers[InAbsoluteIndex];
+}
+
+const FSceneState* FSceneStateExecutionContext::GetState(const FSceneStateMachine& InStateMachine, uint16 InRelativeIndex) const
+{
+	const FSceneStateRange StateRange = InStateMachine.GetStateRange();
+	if (StateRange.Count == 0 || !GeneratedClass)
+	{
+		return nullptr;
+	}
+
+	const TConstArrayView<FSceneState> StatesView = GeneratedClass->GetStates();
+
+	const uint16 AbsoluteIndex = StateRange.Index + InRelativeIndex;
+
+	checkfSlow(StatesView.IsValidIndex(AbsoluteIndex)
+		, TEXT("[Context: %s] State Machine Relative Active Index %d from [%d, %d] is out of bounds! (Total State Count: %d)")
+		, *GetExecutionContextName()
+		, InRelativeIndex
+		, StateRange.Index
+		, StateRange.GetLastIndex()
+		, StatesView.Num());
+
+	return &StatesView[AbsoluteIndex];
+}
+
+const FSceneStateConduit* FSceneStateExecutionContext::GetConduit(const FSceneStateMachine& InStateMachine, uint16 InRelativeIndex) const
+{
+	const FSceneStateRange ConduitRange = InStateMachine.GetConduitRange();
+	if (ConduitRange.Count == 0 || !GeneratedClass)
+	{
+		return nullptr;
+	}
+
+	const TConstArrayView<FSceneStateConduit> ConduitsView = GeneratedClass->GetConduits();
+
+	const uint16 AbsoluteIndex = ConduitRange.Index + InRelativeIndex;
+
+	checkfSlow(ConduitsView.IsValidIndex(AbsoluteIndex)
+		, TEXT("[Context: %s] State machine conduit relative index %d from [%d, %d] is out of bounds! (Total conduit count: %d)")
+		, *GetExecutionContextName()
+		, InRelativeIndex
+		, ConduitRange.Index
+		, ConduitRange.GetLastIndex()
+		, ConduitsView.Num());
+
+	return &ConduitsView[AbsoluteIndex];
+}
+
+FSceneStateInstance* FSceneStateExecutionContext::FindOrAddStateInstance(const FSceneState& InState) const
+{
+	uint16 StateIndex;
+	if (GetStateIndex(InState, /*out*/StateIndex))
+	{
+		return &StateInstances.FindOrAdd(StateIndex);
+	}
+	return nullptr;
+}
+
+FSceneStateInstance* FSceneStateExecutionContext::FindStateInstance(uint16 InAbsoluteIndex) const
+{
+	return StateInstances.Find(InAbsoluteIndex);
+}
+
+FSceneStateInstance* FSceneStateExecutionContext::FindStateInstance(const FSceneState& InState) const
+{
+	uint16 StateIndex;
+	if (GetStateIndex(InState, /*out*/StateIndex))
+	{
+		return FindStateInstance(StateIndex);
+	}
+	return nullptr;
+}
+
+void FSceneStateExecutionContext::RemoveStateInstance(const FSceneState& InState) const
+{
+	uint16 StateIndex;
+	if (GetStateIndex(InState, /*out*/StateIndex))
+	{
+		StateInstances.Remove(StateIndex);
+	}
+}
+
+FConstStructView FSceneStateExecutionContext::FindTask(uint16 InAbsoluteIndex) const
+{
+	if (!GeneratedClass)
+	{
+		return {};
+	}
+
+	const FInstancedStructContainer& Tasks = GeneratedClass->GetTasks();
+
+	if (!Tasks.IsValidIndex(InAbsoluteIndex))
+	{
+		return {};
+	}
+
+	return Tasks[InAbsoluteIndex];
+}
+
+FInstancedStructContainer* FSceneStateExecutionContext::FindOrAddTaskInstanceContainer(const FSceneState& InState) const
+{
+	uint16 StateIndex;
+	if (GetStateIndex(InState, /*out*/StateIndex))
+	{
+		return &TaskInstanceContainers.FindOrAdd(StateIndex);
+	}
+	return nullptr;
+}
+
+FInstancedStructContainer* FSceneStateExecutionContext::FindTaskInstanceContainer(const FSceneState& InState) const
+{
+	uint16 StateIndex;
+	if (GetStateIndex(InState, /*out*/StateIndex))
+	{
+		return FindTaskInstanceContainer(StateIndex);
+	}
+	return nullptr;
+}
+
+FInstancedStructContainer* FSceneStateExecutionContext::FindTaskInstanceContainer(uint16 InAbsoluteIndex) const
+{
+	return TaskInstanceContainers.Find(InAbsoluteIndex);
+}
+
+FStructView FSceneStateExecutionContext::FindTaskInstance(uint16 InAbsoluteIndex) const
+{
+	if (!GeneratedClass)
+	{
+		return {};
+	}
+
+	const FInstancedStructContainer& Tasks = GeneratedClass->GetTasks();
+
+	if (!Tasks.IsValidIndex(InAbsoluteIndex))
+	{
+		return {};
+	}
+
+	const FSceneStateTask& Task = Tasks[InAbsoluteIndex].Get<const FSceneStateTask>();
+
+	const uint16 StateIndex = Task.GetParentStateIndex();
+
+	FInstancedStructContainer* const TaskInstanceContainer = FindTaskInstanceContainer(StateIndex);
+	if (!TaskInstanceContainer)
+	{
+		return {};
+	}
+
+	const TConstArrayView<FSceneState> StatesView = GeneratedClass->GetStates();
+
+	checkf(StatesView.IsValidIndex(StateIndex)
+		, TEXT("[Context: %s] Task Parent State Index %d is out of bounds! (Total State Count: %d)")
+		, *GetExecutionContextName()
+		, StateIndex
+		, StatesView.Num());
+
+	const FSceneState& State = StatesView[StateIndex];
+
+	const uint16 TaskRelativeIndex = InAbsoluteIndex - State.GetTaskRange().Index;
+	if (!TaskInstanceContainer->IsValidIndex(TaskRelativeIndex))
+	{
+		return {};
+	}
+
+	return (*TaskInstanceContainer)[TaskRelativeIndex];
+}
+
+void FSceneStateExecutionContext::RemoveTaskInstanceContainer(const FSceneState& InState) const
+{
+	uint16 StateIndex;
+	if (GetStateIndex(InState, /*out*/StateIndex))
+	{
+		TaskInstanceContainers.Remove(StateIndex);
+	}
+}
+
+FSceneStateMachineInstance* FSceneStateExecutionContext::FindOrAddStateMachineInstance(const FSceneStateMachine& InStateMachine) const
+{
+	uint16 StateMachineIndex;
+	if (GetStateMachineIndex(InStateMachine, /*out*/StateMachineIndex))
+	{
+		return &StateMachineInstances.FindOrAdd(StateMachineIndex);
+	}
+	return nullptr;
+}
+
+FSceneStateMachineInstance* FSceneStateExecutionContext::FindStateMachineInstance(uint16 InAbsoluteIndex) const
+{
+	return StateMachineInstances.Find(InAbsoluteIndex);
+}
+
+FSceneStateMachineInstance* FSceneStateExecutionContext::FindStateMachineInstance(const FSceneStateMachine& InStateMachine) const
+{
+	uint16 StateMachineIndex;
+	if (GetStateMachineIndex(InStateMachine, /*out*/StateMachineIndex))
+	{
+		return FindStateMachineInstance(StateMachineIndex);
+	}
+	return nullptr;
+}
+
+void FSceneStateExecutionContext::RemoveStateMachineInstance(const FSceneStateMachine& InStateMachine) const
+{
+	uint16 StateMachineIndex;
+	if (GetStateMachineIndex(InStateMachine, /*out*/StateMachineIndex))
+	{
+		StateMachineInstances.Remove(StateMachineIndex);
+	}
+}
+
+FSceneStateTransitionInstance* FSceneStateExecutionContext::FindOrAddTransitionInstance(const FSceneStateTransition& InTransition) const
+{
+	uint16 TransitionIndex;
+	if (GetTransitionIndex(InTransition, /*out*/TransitionIndex))
+	{
+		return &TransitionInstances.FindOrAdd(TransitionIndex);
+	}
+	return nullptr;
+}
+
+FSceneStateTransitionInstance* FSceneStateExecutionContext::FindTransitionInstance(uint16 InAbsoluteIndex) const
+{
+	return TransitionInstances.Find(InAbsoluteIndex);
+}
+
+FSceneStateTransitionInstance* FSceneStateExecutionContext::FindTransitionInstance(const FSceneStateTransition& InTransition) const
+{
+	uint16 TransitionIndex;
+	if (GetTransitionIndex(InTransition, /*out*/TransitionIndex))
+	{
+		return FindTransitionInstance(TransitionIndex);
+	}
+	return nullptr;
+}
+
+void FSceneStateExecutionContext::RemoveTransitionInstance(const FSceneStateTransition& InTransition) const
+{
+	uint16 TransitionIndex;
+	if (GetTransitionIndex(InTransition, /*out*/TransitionIndex))
+	{
+		TransitionInstances.Remove(TransitionIndex);
+	}
+}
+
+FSceneStateConduitInstance* FSceneStateExecutionContext::FindOrAddConduitInstance(const FSceneStateConduit& InConduit) const
+{
+	uint16 ConduitIndex;
+	if (GetConduitIndex(InConduit, /*out*/ConduitIndex))
+	{
+		return &ConduitInstances.FindOrAdd(ConduitIndex);
+	}
+	return nullptr;
+}
+
+void FSceneStateExecutionContext::RemoveConduitInstance(const FSceneStateConduit& InConduit) const
+{
+	uint16 ConduitIndex;
+	if (GetConduitIndex(InConduit, /*out*/ConduitIndex))
+	{
+		ConduitInstances.Remove(ConduitIndex);
+	}
+}
+
+const FSceneStateExecutionContext* FSceneStateExecutionContext::FindParentContext() const
+{
+	if (TSharedPtr<UE::SceneState::FExecutionContextRegistry> ContextRegistry = ContextRegistryWeak.Pin())
+	{
+		return ContextRegistry->FindContext(ParentHandle);
+	}
+	return nullptr;
+}
+
+bool FSceneStateExecutionContext::GetStateIndex(const FSceneState& InState, uint16& OutIndex) const
+{
+	if (GeneratedClass)
+	{
+		TConstArrayView<FSceneState> StatesView = GeneratedClass->GetStates();
+		const int32 Index = static_cast<int32>(&InState - StatesView.GetData());
+		checkSlow(StatesView.IsValidIndex(Index) && Index >= 0 && Index <= TNumericLimits<uint16>::Max());
+		OutIndex = static_cast<uint16>(Index);
+		return true;
+	}
+	return false;
+}
+
+bool FSceneStateExecutionContext::GetStateMachineIndex(const FSceneStateMachine& InStateMachine, uint16& OutIndex) const
+{
+	if (GeneratedClass)
+	{
+		TConstArrayView<FSceneStateMachine> StateMachinesView = GeneratedClass->GetStateMachines();
+		const int32 Index = static_cast<int32>(&InStateMachine - StateMachinesView.GetData());
+		checkSlow(StateMachinesView.IsValidIndex(Index) && Index >= 0 && Index <= TNumericLimits<uint16>::Max());
+		OutIndex = static_cast<uint16>(Index);
+		return true;
+	}
+	return false;
+}
+
+bool FSceneStateExecutionContext::GetTransitionIndex(const FSceneStateTransition& InTransition, uint16& OutIndex) const
+{
+	if (GeneratedClass)
+	{
+		TConstArrayView<FSceneStateTransition> TransitionsView = GeneratedClass->GetTransitions();
+		const int32 Index = static_cast<int32>(&InTransition - TransitionsView.GetData());
+		checkSlow(TransitionsView.IsValidIndex(Index) && Index >= 0 && Index <= TNumericLimits<uint16>::Max());
+		OutIndex = static_cast<uint16>(Index);
+		return true;
+	}
+	return false;
+}
+
+bool FSceneStateExecutionContext::GetConduitIndex(const FSceneStateConduit& InConduit, uint16& OutIndex) const
+{
+	if (GeneratedClass)
+	{
+		TConstArrayView<FSceneStateConduit> ConduitsView = GeneratedClass->GetConduits();
+		const int32 Index = static_cast<int32>(&InConduit - ConduitsView.GetData());
+		checkSlow(ConduitsView.IsValidIndex(Index) && Index >= 0 && Index <= TNumericLimits<uint16>::Max());
+		OutIndex = static_cast<uint16>(Index);
+		return true;
+	}
+	return false;
+}
diff --git a/Engine/Plugins/VirtualProduction/SceneState/Source/SceneState/Private/SceneStateGeneratedClass.cpp b/Engine/Plugins/VirtualProduction/SceneState/Source/SceneState/Private/SceneStateGeneratedClass.cpp
new file mode 100644
--- /dev/null
+++ b/Engine/Plugins/VirtualProduction/SceneState/Source/SceneState/Private/SceneStateGeneratedClass.cpp
@@ -0,0 +1,304 @@
+// Copyright Epic Games, Inc. All Rights Reserved.
+
+#include "SceneStateGeneratedClass.h"
+#include "Misc/EnumerateRange.h"
+#include "SceneStateBindingUtils.h"
+#include "SceneStateEventSchema.h"
+#include "SceneStateExecutionContext.h"
+#include "SceneStateExecutionContextRegistry.h"
+#include "SceneStateLog.h"
+#include "SceneStateMachine.h"
+#include "SceneStateObject.h"
+#include "SceneStateUtils.h"
+#include "StructUtils/UserDefinedStruct.h"
+#include "Tasks/SceneStateTask.h"
+#include "Tasks/SceneStateTaskBindingExtension.h"
+#include "Tasks/SceneStateTaskInstance.h"
+
+#if WITH_EDITOR
+#include "StructUtils/UserDefinedStruct.h"
+#include "StructUtilsDelegates.h"
+#endif
+
+USceneStateGeneratedClass::USceneStateGeneratedClass()
+{
+#if WITH_EDITOR
+	if (!HasAnyFlags(RF_ClassDefaultObject))
+	{
+		OnObjectsReinstancedHandle = FCoreUObjectDelegates::OnObjectsReinstanced.AddUObject(this, &USceneStateGeneratedClass::OnObjectsReinstanced);
+		OnStructsReinstancedHandle = UE::StructUtils::Delegates::OnUserDefinedStructReinstanced.AddUObject(this, &USceneStateGeneratedClass::OnStructsReinstanced);
+	}
+#endif
+}
+
+const FSceneState* USceneStateGeneratedClass::GetRootState() const
+{
+	if (States.IsValidIndex(RootStateIndex))
+	{
+		return &States[RootStateIndex];	
+	}
+	return nullptr;
+}
+
+const FSceneStateMachine* USceneStateGeneratedClass::FindStateMachine(const FGuid& InStateMachineId) const
+{
+	const uint16* StateMachineIndex = StateMachineIdToIndex.Find(InStateMachineId);
+	if (StateMachineIndex && StateMachines.IsValidIndex(*StateMachineIndex))
+	{
+		return &StateMachines[*StateMachineIndex];
+	}
+	return nullptr;
+}
+
+FInstancedPropertyBag USceneStateGeneratedClass::GetTransitionParameters(uint16 InAbsoluteIndex) const
+{
+	return TransitionParameters.FindRef(InAbsoluteIndex);
+}
+
+void USceneStateGeneratedClass::Link(FArchive& Ar, bool bInRelinkExistingProperties)
+{
+	Super::Link(Ar, bInRelinkExistingProperties);
+
+	if (IsFullClass())
+	{
+		for (TEnumerateRef<FSceneStateTransition> Transition : EnumerateRange(Transitions))
+		{
+			if (TransitionLinks.IsValidIndex(Transition.GetIndex()))
+			{
+				Transition->Link(TransitionLinks[Transition.GetIndex()], this);
+			}
+		}
+
+		for (TEnumerateRef<FSceneStateConduit> Conduit : EnumerateRange(Conduits))
+		{
+			if (ConduitLinks.IsValidIndex(Conduit.GetIndex()))
+			{
+				Conduit->Link(ConduitLinks[Conduit.GetIndex()], this);
+			}
+		}
+	}
+}
+
+void USceneStateGeneratedClass::PurgeClass(bool bInRecompilingOnLoad)
+{
+	Super::PurgeClass(bInRecompilingOnLoad);
+	Reset();
+}
+
+void USceneStateGeneratedClass::PostLoad()
+{
+	Super::PostLoad();
+	ResolveBindings();
+}
+
+void USceneStateGeneratedClass::BeginDestroy()
+{
+	Super::BeginDestroy();
+	Reset();
+}
+
+#if WITH_EDITOR
+const FSceneState* USceneStateGeneratedClass::FindStateFromNode(FObjectKey InStateNode) const
+{
+	const uint16* StateIndex = StateNodeToIndex.Find(InStateNode);
+	if (StateIndex && States.IsValidIndex(*StateIndex))
+	{
+		return &States[*StateIndex];
+	}
+	return nullptr;
+}
+
+const FSceneStateTask* USceneStateGeneratedClass::FindTaskFromNode(FObjectKey InTaskNode) const
+{
+	const uint16* TaskIndex = TaskNodeToIndex.Find(InTaskNode);
+	if (TaskIndex && Tasks.IsValidIndex(*TaskIndex))
+	{
+		return Tasks[*TaskIndex].GetPtr<const FSceneStateTask>();
+	}
+	return nullptr;
+}
+
+void USceneStateGeneratedClass::ForEachStateInstance(USceneStateObject& InRootState, FObjectKey InStateNode, TFunctionRef<void(const FSceneStateInstance&)> InFunctor) const
+{
+	const uint16* StateIndex = StateNodeToIndex.Find(InStateNode);
+	if (!StateIndex)
+	{
+		return;
+	}
+
+	const TSharedRef<UE::SceneState::FExecutionContextRegistry> ContextRegistry = InRootState.GetContextRegistry();
+
+	ContextRegistry->ForEachExecutionContext(
+		[StateIndex, &InFunctor](const FSceneStateExecutionContext& InExecutionContext)
+		{
+			if (const FSceneStateInstance* StateInstance = InExecutionContext.FindStateInstance(*StateIndex))
+			{
+				InFunctor(*StateInstance);
+			}
+		});
+}
+
+void USceneStateGeneratedClass::ForEachTaskInstance(USceneStateObject& InRootState, FObjectKey InTaskNode, TFunctionRef<void(const FSceneStateTaskInstance&)> InFunctor) const
+{
+	const uint16* TaskIndex = TaskNodeToIndex.Find(InTaskNode);
+	if (!TaskIndex)
+	{
+		return;
+	}
+
+	const TSharedRef<UE::SceneState::FExecutionContextRegistry> ContextRegistry = InRootState.GetContextRegistry();
+
+	ContextRegistry->ForEachExecutionContext(
+		[TaskIndex, &InFunctor](const FSceneStateExecutionContext& InExecutionContext)
+		{
+			if (const FSceneStateTaskInstance* TaskInstance = InExecutionContext.FindTaskInstance(*TaskIndex).GetPtr<const FSceneStateTaskInstance>())
+			{
+				InFunctor(*TaskInstance);
+			}
+		});
+}
+#endif
+
+const UStruct* USceneStateGeneratedClass::FindDataStruct(const FSceneStateBindingDataHandle& InDataHandle)
+{
+	const uint16 DataIndex = InDataHandle.GetDataIndex();
+	const uint16 DataSubIndex = InDataHandle.GetDataSubIndex();
+
+	if (InDataHandle.IsExternalDataType())
+	{
+		// Todo: unsupported external sources
+		return nullptr;
+	}
+
+	const ESceneStateDataType DataType = static_cast<ESceneStateDataType>(InDataHandle.GetDataType());
+
+	switch (DataType)
+	{
+	case ESceneStateDataType::Root:
+		return this;
+
+	case ESceneStateDataType::Task:
+		return TaskInstances[DataIndex].GetScriptStruct();
+
+	case ESceneStateDataType::EventHandler:
+		return EventHandlers[DataIndex].GetEventSchemaHandle().GetEventStruct();
+
+	case ESceneStateDataType::TaskExtension:
+		if (const FSceneStateTask* Task = Tasks[DataIndex].GetPtr<const FSceneStateTask>())
+		{
+			if (const FSceneStateTaskBindingExtension* BindingExtension = Task->GetBindingExtension())
+			{
+				FStructView DataView;
+				if (BindingExtension->FindDataByIndex(TaskInstances[DataIndex], DataSubIndex, DataView))
+				{
+					return DataView.GetScriptStruct();
+				}
+			}
+		}
+		break;
+
+	case ESceneStateDataType::Transition:
+		return TransitionParameters[DataIndex].GetPropertyBagStruct();
+
+	case ESceneStateDataType::StateMachine:
+		return StateMachines[DataIndex].GetParametersStruct();
+	}
+
+	return nullptr;
+}
+
+void USceneStateGeneratedClass::ResolveBindings()
+{
+	if (!IsFullClass())
+	{
+		return;	
+	}
+
+	UE::SceneState::PatchBindingCollection({ .BindingCollection = BindingCollection, .FindDataStructFunctor =
+		[this](const FSceneStateBindingDataHandle& InDataHandle)->const UStruct*
+		{
+			return FindDataStruct(InDataHandle);
+		}});
+
+	// Resolves property paths used by bindings a store property pointers
+	if (!BindingCollection.ResolvePaths())
+	{
+		UE_LOG(LogSceneState, Warning, TEXT("%s failed to resolve bindings. Try compiling Scene State Blueprint again"), *GetName());
+	}
+}
+
+void USceneStateGeneratedClass::Reset()
+{
+	Tasks.Reset();
+	TaskPrerequisites.Reset();
+	TaskInstances.Reset();
+	States.Reset();
+	Conduits.Reset();
+	ConduitLinks.Reset();
+	EventHandlers.Reset();
+	Transitions.Reset();
+	TransitionLinks.Reset();
+	TransitionParameters.Reset();
+	StateMachines.Reset();
+	BindingCollection.Reset();
+	StateMachineIdToIndex.Reset();
+
+#if WITH_EDITOR
+	TaskMetadata.Reset();
+	StateMetadata.Reset();
+	TransitionMetadata.Reset();
+	StateNodeToIndex.Reset();
+	StateMachineGraphToIndex.Reset();
+	TaskNodeToIndex.Reset();
+#endif
+}
+
+bool USceneStateGeneratedClass::IsFullClass() const
+{
+#if WITH_EDITOR
+	FNameBuilder ClassName(GetFName());
+	FStringView ClassNameView = ClassName.ToView();
+	return !ClassNameView.StartsWith(TEXT("SKEL_")) && !ClassNameView.StartsWith(TEXT("REINST_"));
+#else
+	return true;
+#endif
+}
+
+#if WITH_EDITOR
+void USceneStateGeneratedClass::OnObjectsReinstanced(const TMap<UObject*, UObject*>& InReplacementMap)
+{
+	TSet<const UStruct*> Structs;
+	Structs.Reserve(InReplacementMap.Num());
+
+	bool bRequiresResolve = false;
+
+	for (const TPair<UObject*, UObject*>& Pair : InReplacementMap)
+	{
+		if (const UObject* Replacement = Pair.Value)
+		{
+			if (Replacement->IsIn(this))
+			{
+				bRequiresResolve = true;
+				break;
+			}
+			Structs.Add(Replacement->GetClass());
+		}
+	}
+
+	if (bRequiresResolve || BindingCollection.ContainsAnyStruct(Structs))
+	{
+		ResolveBindings();
+	}
+}
+
+void USceneStateGeneratedClass::OnStructsReinstanced(const UUserDefinedStruct& InStruct)
+{
+	TSet<const UStruct*> Structs;
+	Structs.Add(&InStruct);
+
+	if (BindingCollection.ContainsAnyStruct(Structs))
+	{
+		ResolveBindings();
+	}
+}
+#endif
diff --git a/Engine/Plugins/VirtualProduction/SceneState/Source/SceneState/Private/Tasks/SceneStateTask.cpp b/Engine/Plugins/VirtualProduction/SceneState/Source/SceneState/Private/Tasks/SceneStateTask.cpp
new file mode 100644
--- /dev/null
+++ b/Engine/Plugins/VirtualProduction/SceneState/Source/SceneState/Private/Tasks/SceneStateTask.cpp
@@ -0,0 +1,197 @@
+// Copyright Epic Games, Inc. All Rights Reserved.
+
+#include "Tasks/SceneStateTask.h"
+#include "SceneStateBindingUtils.h"
+#include "SceneStateExecutionContext.h"
+#include "SceneStateObject.h"
+#include "SceneStateUtils.h"
+#include "Tasks/SceneStateTaskBindingExtension.h"
+#include "Tasks/SceneStateTaskInstance.h"
+
+namespace UE::SceneState::Private
+{
+
+EExecutionStatus GetTaskStatus(ESceneStateTaskStopReason InStopReason)
+{
+	switch (InStopReason)
+	{
+	case ESceneStateTaskStopReason::State:
+		return EExecutionStatus::NotStarted;
+
+	case ESceneStateTaskStopReason::Finished:
+		return EExecutionStatus::Finished;
+	}
+
+	checkNoEntry();
+	return EExecutionStatus::Finished;
+}
+
+} // UE::SceneState::Private
+
+#if WITH_EDITOR
+const UScriptStruct* FSceneStateTask::GetTaskInstanceType() const
+{
+	return OnGetTaskInstanceType();
+}
+
+void FSceneStateTask::BuildTaskInstance(UObject* InOuter, FStructView InTaskInstance) const
+{
+	if (InTaskInstance.IsValid())
+	{
+		OnBuildTaskInstance(InOuter, InTaskInstance);
+	}
+}
+
+void FSceneStateTask::PostEditChange(UE::SceneState::FTaskEditChange& InEditChange, FStructView InTaskInstance)
+{
+	if (InTaskInstance.IsValid())
+	{
+		OnPostEditChange(InEditChange, InTaskInstance);
+	}
+}
+#endif
+
+const FSceneStateTaskBindingExtension* FSceneStateTask::GetBindingExtension() const
+{
+	if (EnumHasAllFlags(TaskFlags, ESceneStateTaskFlags::HasBindingExtension))
+	{
+		return OnGetBindingExtension();
+	}
+	return nullptr;
+}
+
+FStructView FSceneStateTask::FindTaskInstance(const FSceneStateExecutionContext& InContext) const
+{
+	return InContext.FindTaskInstance(GetTaskIndex());
+}
+
+void FSceneStateTask::Setup(const FSceneStateExecutionContext& InContext, FStructView InTaskInstance) const
+{
+	using namespace UE::SceneState;
+
+	QUICK_SCOPE_CYCLE_COUNTER(SceneState_OnSetup);
+
+	FSceneStateTaskInstance* Instance = InTaskInstance.GetPtr<FSceneStateTaskInstance>();
+	if (!Instance)
+	{
+		return;
+	}
+
+	Instance->SetStatus(EExecutionStatus::NotStarted);
+	OnSetup(InContext, InTaskInstance);
+}
+
+void FSceneStateTask::Start(const FSceneStateExecutionContext& InContext, FStructView InTaskInstance) const
+{
+	using namespace UE::SceneState;
+
+	QUICK_SCOPE_CYCLE_COUNTER(SceneState_OnStart);
+
+	FSceneStateTaskInstance* Instance = InTaskInstance.GetPtr<FSceneStateTaskInstance>();
+	if (!Instance || Instance->GetStatus() != EExecutionStatus::NotStarted)
+	{
+		return;
+	}
+
+	Instance->SetStatus(EExecutionStatus::Running);
+
+	ApplyBindings(InContext, InTaskInstance);
+
+	OnStart(InContext, InTaskInstance);
+}
+
+void FSceneStateTask::Tick(const FSceneStateExecutionContext& InContext, FStructView InTaskInstance, float InDeltaSeconds) const
+{
+	using namespace UE::SceneState;
+
+	// Skip task is not set to tick
+	if (!EnumHasAnyFlags(TaskFlags, ESceneStateTaskFlags::Ticks))
+	{
+		return;
+	}
+
+	QUICK_SCOPE_CYCLE_COUNTER(SceneState_OnTick);
+
+	FSceneStateTaskInstance* Instance = InTaskInstance.GetPtr<FSceneStateTaskInstance>();
+	if (!Instance || Instance->GetStatus() != EExecutionStatus::Running)
+	{
+		return;
+	}
+
+	OnTick(InContext, InTaskInstance, InDeltaSeconds);
+}
+
+void FSceneStateTask::Stop(const FSceneStateExecutionContext& InContext, FStructView InTaskInstance, ESceneStateTaskStopReason InStopReason) const
+{
+	using namespace UE::SceneState;
+
+	QUICK_SCOPE_CYCLE_COUNTER(SceneState_OnStop);
+
+	FSceneStateTaskInstance* Instance = InTaskInstance.GetPtr<FSceneStateTaskInstance>();
+	if (!Instance)
+	{
+		return;
+	}
+
+	if (Instance->GetStatus() == EExecutionStatus::Running)
+	{
+		OnStop(InContext, InTaskInstance, InStopReason);
+	}
+
+	Instance->SetStatus(Private::GetTaskStatus(InStopReason));
+}
+
+void FSceneStateTask::Finish(const FSceneStateExecutionContext& InContext, FStructView InTaskInstance) const
+{
+	Stop(InContext, InTaskInstance, ESceneStateTaskStopReason::Finished);
+}
+
+void FSceneStateTask::SetFlags(ESceneStateTaskFlags InFlags)
+{
+	TaskFlags |= InFlags;
+}
+
+void FSceneStateTask::ClearFlags(ESceneStateTaskFlags InFlags)
+{
+	TaskFlags &= ~InFlags;
+}
+
+bool FSceneStateTask::ApplyBindings(const FSceneStateExecutionContext& InContext, FStructView InTaskInstance) const
+{
+	QUICK_SCOPE_CYCLE_COUNTER(SceneStateTask_ApplyBindings);
+
+	const FSceneStateBindingCollection* const BindingCollection = InContext.GetBindingCollection();
+	if (!BindingCollection)
+	{
+		return false;
+	}
+
+	auto FindDataView =
+		[&InContext](const FSceneStateBindingDataHandle& InDataHandle)
+		{
+			return InContext.FindDataView(InDataHandle);
+		};
+
+	UE::SceneState::FApplyBatchParams ApplyBatchParams
+		{
+			.BatchIndex = BindingsBatch.Get(),
+			.TargetDataView = InTaskInstance,
+			.BindingCollection = *BindingCollection,
+			.FindSourceDataViewFunctor = FindDataView,
+		};
+
+	bool bResult = ApplyBatch(ApplyBatchParams);
+
+	if (const FSceneStateTaskBindingExtension* BindingExtension = GetBindingExtension())
+	{
+		BindingExtension->VisitBindingBatches(InTaskInstance,
+			[&bResult, &ApplyBatchParams](uint16 InBatchIndex, FStructView InTargetDataView)
+			{
+				ApplyBatchParams.BatchIndex = InBatchIndex;
+				ApplyBatchParams.TargetDataView = InTargetDataView;
+				bResult &= ApplyBatch(ApplyBatchParams);
+			});
+	}
+
+	return bResult;
+}
diff --git a/Engine/Plugins/VirtualProduction/SceneState/Source/SceneState/Private/Tasks/SceneStateTaskInstance.cpp b/Engine/Plugins/VirtualProduction/SceneState/Source/SceneState/Private/Tasks/SceneStateTaskInstance.cpp
new file mode 100644
--- /dev/null
+++ b/Engine/Plugins/VirtualProduction/SceneState/Source/SceneState/Private/Tasks/SceneStateTaskInstance.cpp
@@ -0,0 +1,8 @@
+// Copyright Epic Games, Inc. All Rights Reserved.
+
+#include "Tasks/SceneStateTaskInstance.h"
+
+void FSceneStateTaskInstance::SetStatus(UE::SceneState::EExecutionStatus InStatus)
+{
+	Status = InStatus;
+}
diff --git a/Engine/Plugins/VirtualProduction/SceneState/Source/SceneState/Public/SceneStateGeneratedClass.h b/Engine/Plugins/VirtualProduction/SceneState/Source/SceneState/Public/SceneStateGeneratedClass.h
new file mode 100644
--- /dev/null
+++ b/Engine/Plugins/VirtualProduction/SceneState/Source/SceneState/Public/SceneStateGeneratedClass.h
@@ -0,0 +1,251 @@
+// Copyright Epic Games, Inc. All Rights Reserved.
+
+#pragma once
+
+#include "Conduit/SceneStateConduit.h"
+#include "Conduit/SceneStateConduitLink.h"
+#include "Engine/BlueprintGeneratedClass.h"
+#include "SceneState.h"
+#include "SceneStateBindingCollection.h"
+#include "SceneStateEventHandler.h"
+#include "SceneStateMachine.h"
+#include "SceneStateRange.h"
+#include "StructUtils/InstancedStructContainer.h"
+#include "Tasks/SceneStateTask.h"
+#include "Tasks/SceneStateTaskMetadata.h"
+#include "Transition/SceneStateTransition.h"
+#include "Transition/SceneStateTransitionLink.h"
+#include "Transition/SceneStateTransitionMetadata.h"
+#include "UObject/ObjectKey.h"
+#include "UObject/Package.h"
+#include "SceneStateGeneratedClass.generated.h"
+
+class USceneStateObject;
+class UUserDefinedStruct;
+
+namespace UE::SceneState::Editor
+{
+	class FBindingCompiler;
+	class FBlueprintCompilerContext;
+	class FStateMachineCompiler;
+}
+
+/**
+ * Object Class for the Scene State Object
+ * Holds all the data about States, State Machines, Tasks, etc.
+ * All this data is immutable in execution, and as such, it is not instanced to the Scene State Object instances.
+ * @see FSceneStateExecutionContext
+ */
+UCLASS(MinimalAPI)
+class USceneStateGeneratedClass : public UBlueprintGeneratedClass
+{
+	GENERATED_BODY()
+
+public:
+	SCENESTATE_API USceneStateGeneratedClass();
+
+	/** Gets the state that owns all the "top-level" state machines */
+	const FSceneState* GetRootState() const;
+
+	/** Finds the Top-Level State Machine that is linked to the given Guid */
+	const FSceneStateMachine* FindStateMachine(const FGuid& InStateMachineId) const;
+
+	const FSceneStateBindingCollection& GetBindingCollection() const
+	{
+		return BindingCollection;
+	}
+
+	TConstArrayView<FSceneState> GetStates() const
+	{
+		return States;
+	}
+
+	TConstArrayView<FSceneStateConduit> GetConduits() const
+	{
+		return Conduits;
+	}
+
+	const FInstancedStructContainer& GetTasks() const
+	{
+		return Tasks;
+	}
+
+	TConstArrayView<uint16> GetTaskPrerequisites() const
+	{
+		return TaskPrerequisites;
+	}
+
+	const FInstancedStructContainer& GetTaskInstances() const
+	{
+		return TaskInstances;
+	}
+
+	TConstArrayView<FSceneStateEventHandler> GetEventHandlers() const
+	{
+		return EventHandlers;
+	}
+
+	TConstArrayView<FSceneStateTransition> GetTransitions() const
+	{
+		return Transitions;
+	}
+
+	/** Gets a copy of the transition parameters at the given transition absolute index. */
+	FInstancedPropertyBag GetTransitionParameters(uint16 InAbsoluteIndex) const;
+
+	TConstArrayView<FSceneStateMachine> GetStateMachines() const
+	{
+		return StateMachines;
+	}
+
+#if WITH_EDITOR
+	TConstArrayView<FSceneStateMetadata> GetStateMetadata() const
+	{
+		return StateMetadata;
+	}
+
+	TConstArrayView<FSceneStateTaskMetadata> GetTaskMetadata() const
+	{
+		return TaskMetadata;
+	}
+#endif
+
+	//~ Begin UClass
+	SCENESTATE_API virtual void Link(FArchive& Ar, bool bInRelinkExistingProperties) override;
+	SCENESTATE_API virtual void PurgeClass(bool bInRecompilingOnLoad) override;
+	//~ End UClass
+
+	//~ Begin UObject
+	SCENESTATE_API virtual void PostLoad() override;
+	SCENESTATE_API virtual void BeginDestroy() override;
+	//~ End UObject
+
+#if WITH_EDITOR
+	/** Finds the mapped compiled state for a given state node */
+	SCENESTATE_API const FSceneState* FindStateFromNode(FObjectKey InStateNode) const;
+
+	/** Finds the mapped compiled task for a given task node */
+	SCENESTATE_API const FSceneStateTask* FindTaskFromNode(FObjectKey InTaskNode) const;
+
+	/**
+	 * For a given root state and state node retrieves the mapped state instances
+	 * @param InRootState the root state that is outered to a player
+	 * @param InStateNode the state node to look for 
+	 * @param InFunctor the functor to run for each state instance
+	 */
+	SCENESTATE_API void ForEachStateInstance(USceneStateObject& InRootState
+		, FObjectKey InStateNode
+		, TFunctionRef<void(const FSceneStateInstance&)> InFunctor) const;
+
+	/**
+	 * For a given root state and task node retrieves the mapped task instances
+	 * @param InRootState the root state that is outered to a player
+	 * @param InTaskNode the task node to look for
+	 * @param InFunctor the functor to run for each task instance
+	 */
+	SCENESTATE_API void ForEachTaskInstance(USceneStateObject& InRootState
+		, FObjectKey InTaskNode
+		, TFunctionRef<void(const FSceneStateTaskInstance&)> InFunctor) const;
+#endif
+
+private:
+	/** Finds the struct type for the given data handle */
+	const UStruct* FindDataStruct(const FSceneStateBindingDataHandle& InDataHandle);
+
+	/** Patches bindings and resolves the binding paths for the owning binding collection */
+	SCENESTATE_API void ResolveBindings();
+
+	/** Resets all the elements that get compiled for this generated class (e.g. Tasks, State Machines, etc) */
+	void Reset();
+
+	bool IsFullClass() const;
+
+#if WITH_EDITOR
+	void OnObjectsReinstanced(const TMap<UObject*, UObject*>& InReplacementMap);
+	void OnStructsReinstanced(const UUserDefinedStruct& InStruct);
+#endif
+
+	UPROPERTY()
+	FSceneStateBindingCollection BindingCollection;
+
+	UPROPERTY()
+	uint16 RootStateIndex = FSceneStateRange::InvalidIndex;
+
+	/** Holds all the compiled states including nested states */
+	UPROPERTY()
+	TArray<FSceneState> States;
+
+	/** Holds all the compiled tasks (including tasks from nested states) in a contiguous block of memory */
+	UPROPERTY()
+	FInstancedStructContainer Tasks;
+
+	/** Array of the task prerequisites in their relative index. Each task has a unique range indicating their pre-requisites */
+	UPROPERTY()
+	TArray<uint16> TaskPrerequisites;
+
+	/** Templates used to instantiate the Task Instances */
+	UPROPERTY()
+	FInstancedStructContainer TaskInstances;
+
+	/** Holds all the compiled event handlers (including from nested sources) contiguously */
+	UPROPERTY()
+	TArray<FSceneStateEventHandler> EventHandlers;
+
+	/** Holds all the exit transitions (including nested sources) contiguously */
+	UPROPERTY()
+	TArray<FSceneStateTransition> Transitions;
+
+	/** Compiled transition information only used in Link time */
+	UPROPERTY()
+	TArray<FSceneStateTransitionLink> TransitionLinks;
+
+	/** Map of the Transition Index (absolute) to the template transition parameters map for evaluation function call. */
+	UPROPERTY()
+	TMap<uint16, FInstancedPropertyBag> TransitionParameters;
+
+	/** All the compiled conduits */
+	UPROPERTY()
+	TArray<FSceneStateConduit> Conduits;
+
+	/** Compiled conduit information used only in Link time */
+	UPROPERTY()
+	TArray<FSceneStateConduitLink> ConduitLinks;
+
+	UPROPERTY()
+	TArray<FSceneStateMachine> StateMachines;
+
+	/** Map of the Top-Level State machine Parameters id to the index in the state machine array */
+	UPROPERTY()
+	TMap<FGuid, uint16> StateMachineIdToIndex;
+
+#if WITH_EDITORONLY_DATA
+	/** Holds all the state editor-only metadata. This array matches in index with the states array. */
+	UPROPERTY()
+	TArray<FSceneStateMetadata> StateMetadata;
+
+	/** Holds all the task editor-only metadata. This array matches in index with the tasks array. */
+	UPROPERTY()
+	TArray<FSceneStateTaskMetadata> TaskMetadata;
+
+	/** Holds all the transition editor-only metadata. This array matches in index with the transitions array. */
+	UPROPERTY()
+	TArray<FSceneStateTransitionMetadata> TransitionMetadata;
+
+	/** Map of a State Node to its State Index */
+	TMap<FObjectKey, uint16> StateNodeToIndex;
+
+	/** Map of a State Machine Graph to its State Machine Index */
+	TMap<FObjectKey, uint16> StateMachineGraphToIndex;
+
+	/** Map of a Task Node to its Task Index */
+	TMap<FObjectKey, uint16> TaskNodeToIndex;
+
+	FDelegateHandle OnObjectsReinstancedHandle;
+	FDelegateHandle OnStructsReinstancedHandle;
+#endif
+
+	friend UE::SceneState::Editor::FBindingCompiler;
+	friend UE::SceneState::Editor::FBlueprintCompilerContext;
+	friend UE::SceneState::Editor::FStateMachineCompiler;
+	friend class USceneStateBlueprint;
+};
diff --git a/Engine/Plugins/VirtualProduction/SceneState/Source/SceneState/Public/SceneStateUtils.h b/Engine/Plugins/VirtualProduction/SceneState/Source/SceneState/Public/SceneStateUtils.h
new file mode 100644
--- /dev/null
+++ b/Engine/Plugins/VirtualProduction/SceneState/Source/SceneState/Public/SceneStateUtils.h
@@ -0,0 +1,85 @@
+// Copyright Epic Games, Inc. All Rights Reserved.
+
+#pragma once
+
+#include "Containers/Map.h"
+#include "HAL/Platform.h"
+#include "Templates/Function.h"
+
+class UClass;
+class UObject;
+struct FConstStructView;
+struct FInstancedStructContainer;
+struct FPropertyVisitorPath;
+struct FSceneStateRange;
+struct FStructView;
+template<typename T> struct TObjectPtr;
+
+namespace UE::SceneState
+{
+	/**
+	 * Converts a map of keys to relative indices to a map of keys to absolute indices
+	 * @param InOutMap the map to convert
+	 * @param InIndexOffset the index to add to these relative indices to transform them to absolute
+	 */
+	template<typename InKeyType, typename InIndexType = uint16>
+	void ToAbsoluteIndexMap(TMap<InKeyType, InIndexType>& InOutMap, InIndexType InIndexOffset)
+	{
+		for (TPair<InKeyType, InIndexType>& Pair : InOutMap)
+		{
+			Pair.Value += InIndexOffset;
+		}
+	}
+
+	SCENESTATE_API TArray<FStructView> GetStructViews(FInstancedStructContainer& InStructContainer, FSceneStateRange InRange);
+
+	SCENESTATE_API TArray<FStructView> GetStructViews(FInstancedStructContainer& InStructContainer);
+
+	SCENESTATE_API TArray<FConstStructView> GetConstStructViews(const FInstancedStructContainer& InStructContainer, FSceneStateRange InRange);
+
+	SCENESTATE_API TArray<FConstStructView> GetConstStructViews(const FInstancedStructContainer& InStructContainer);
+
+	/**
+	 * Given a FPropertyVisitorPath, attempt to resolve that to a valid data pointer.
+	 * RootObject is required to implement ResolveVisitedPathInfo to provide the resolver logic.
+	 * TODO: FPropertyVisitorPath::GetPropertyDataPtr only takes a UObject ptr.
+	 * @see FPropertyVisitorPath::GetPropertyDataPtr
+	 */
+	SCENESTATE_API void* ResolveVisitedPath(const UScriptStruct* InRootObject, void* InRootData, const FPropertyVisitorPath& InPath);
+
+	/**
+	 * Discards a given object by changing its outer to transient package and adding a prefix to it while also ensuring uniqueness
+	 * @param InObjectToDiscard the object to discard
+	 */
+	SCENESTATE_API void DiscardObject(UObject* InObjectToDiscard);
+
+	/**
+	 * Discards an object with a given name in the given outer,
+	 * by renaming the outer to the transient package and adding a prefix to it while also ensuring uniqueness
+	 * @param InOuter outer of the object. Must be valid.
+	 * @param InObjectName name to use. Must match InOutObject's name if not null
+	 * @param InOnPreDiscardOldObject callback for when the old object is about to be trashed
+	 * @return the object that was discarded, or null if none was found
+	 */
+	SCENESTATE_API UObject* DiscardObject(UObject* InOuter
+		, const TCHAR* InObjectName
+		, TFunctionRef<void(UObject*)> InOnPreDiscardOldObject = [](UObject*){});
+
+	/**
+	 * Replaces a given Object with a new object with the same name but different class
+	 * @param InOutObject object to replace. Can come in as null
+	 * @param InOuter outer of the object. Must be valid.
+	 * @param InClass new class of the new object that will replace the older one
+	 * @param InObjectName name to use. Must match InOutObject's name if not null
+	 * @param InContextName context for debugging purposes
+	 * @param InOnPreDiscardOldObject callback for when the old object is about to be trashed
+	 * @return true if the operation took place, false otherwise
+	 */
+	SCENESTATE_API bool ReplaceObject(UObject*& InOutObject
+		, UObject* InOuter
+		, UClass* InClass
+		, const TCHAR* InObjectName
+		, const TCHAR* InContextName = nullptr
+		, TFunctionRef<void(UObject*)> InOnPreDiscardOldObject = [](UObject*){});
+
+} // UE::SceneState
diff --git a/Engine/Plugins/VirtualProduction/SceneState/Source/SceneState/Public/Tasks/SceneStateTaskInstance.h b/Engine/Plugins/VirtualProduction/SceneState/Source/SceneState/Public/Tasks/SceneStateTaskInstance.h
new file mode 100644
--- /dev/null
+++ b/Engine/Plugins/VirtualProduction/SceneState/Source/SceneState/Public/Tasks/SceneStateTaskInstance.h
@@ -0,0 +1,28 @@
+// Copyright Epic Games, Inc. All Rights Reserved.
+
+#pragma once
+
+#include "SceneStateEnums.h"
+#include "SceneStateTaskInstance.generated.h"
+
+struct FPropertyBindingDataView;
+
+/** Represents the Instance Data of a Task */
+USTRUCT()
+struct FSceneStateTaskInstance
+{
+	GENERATED_BODY()
+
+	/** Sets the execution status of the Task */
+	void SetStatus(UE::SceneState::EExecutionStatus InStatus);
+
+	/** Gets the execution status of the task */
+	UE::SceneState::EExecutionStatus GetStatus() const
+	{
+		return Status;
+	}
+
+private:
+	/** Current status of the Task */
+	UE::SceneState::EExecutionStatus Status = UE::SceneState::EExecutionStatus::NotStarted;
+};
diff --git a/Engine/Plugins/VirtualProduction/SceneState/Source/SceneStateBinding/Public/SceneStateBindingCollection.h b/Engine/Plugins/VirtualProduction/SceneState/Source/SceneStateBinding/Public/SceneStateBindingCollection.h
new file mode 100644
--- /dev/null
+++ b/Engine/Plugins/VirtualProduction/SceneState/Source/SceneStateBinding/Public/SceneStateBindingCollection.h
@@ -0,0 +1,87 @@
+// Copyright Epic Games, Inc. All Rights Reserved.
+
+#pragma once
+
+#include "PropertyBindingBindingCollection.h"
+#include "SceneStateBinding.h"
+#include "SceneStateBindingDesc.h"
+#include "SceneStateBindingReference.h"
+#include "SceneStateBindingCollection.generated.h"
+
+struct FSceneStatePropertyReference;
+
+namespace UE::SceneState::Editor
+{
+	class FBindingCompiler;
+}
+
+/** Holds all the property bindings to use in execution */
+USTRUCT()
+struct FSceneStateBindingCollection : public FPropertyBindingBindingCollection
+{
+	GENERATED_BODY()
+
+	TArrayView<FSceneStateBindingDesc> GetMutableBindingDescs()
+	{
+		return BindingDescs;
+	}
+
+	TArrayView<FSceneStateBinding> GetMutableBindings()
+	{
+		return Bindings;
+	}
+
+	/** Finds the binding desc matching the given data handle */
+	SCENESTATEBINDING_API const FSceneStateBindingDesc* FindBindingDesc(FSceneStateBindingDataHandle InDataHandle) const;
+
+	/** Retrieves the resolved reference for the given property reference */
+	SCENESTATEBINDING_API const FSceneStateBindingResolvedReference* FindResolvedReference(const FSceneStatePropertyReference& InPropertyReference) const;
+
+	/**
+	 * Resolves the property access for the given reference and data view
+	 * @note This does not deal with recursive resolution (for cases where the reference is to another reference).
+	 * Prefer using UE::SceneState::ResolveProperty instead.
+	 */
+	SCENESTATEBINDING_API uint8* ResolveProperty(const FSceneStateBindingResolvedReference& InResolvedReference, FPropertyBindingDataView InDataView) const;
+
+	//~ Begin FPropertyBindingBindingCollection
+#if WITH_EDITOR
+	SCENESTATEBINDING_API virtual FPropertyBindingBinding* AddBindingInternal(const FPropertyBindingPath& InSourcePath, const FPropertyBindingPath& InTargetPath) override;
+	SCENESTATEBINDING_API virtual void RemoveBindingsInternal(TFunctionRef<bool(FPropertyBindingBinding&)> InPredicate) override;
+	SCENESTATEBINDING_API virtual bool HasBindingInternal(TFunctionRef<bool(const FPropertyBindingBinding&)> InPredicate) const override;
+	SCENESTATEBINDING_API virtual const FPropertyBindingBinding* FindBindingInternal(TFunctionRef<bool(const FPropertyBindingBinding&)> InPredicate) const override;
+#endif
+	SCENESTATEBINDING_API virtual int32 GetNumBindings() const override;
+	SCENESTATEBINDING_API virtual int32 GetNumBindableStructDescriptors() const override;
+	SCENESTATEBINDING_API virtual const FPropertyBindingBindableStructDescriptor* GetBindableStructDescriptorFromHandle(FConstStructView InSourceHandleView) const override;
+	SCENESTATEBINDING_API virtual void ForEachBinding(TFunctionRef<void(const FPropertyBindingBinding& Binding)> InFunction) const override;
+	SCENESTATEBINDING_API virtual void ForEachBinding(FPropertyBindingIndex16 InBegin, FPropertyBindingIndex16 InEnd, TFunctionRef<void(const FPropertyBindingBinding& Binding, const int32 BindingIndex)> InFunction) const override;
+	SCENESTATEBINDING_API virtual void ForEachMutableBinding(TFunctionRef<void(FPropertyBindingBinding& Binding)> InFunction) override;
+	SCENESTATEBINDING_API virtual void VisitBindings(TFunctionRef<EVisitResult(const FPropertyBindingBinding& Binding)> InFunction) const override;
+	SCENESTATEBINDING_API virtual void VisitMutableBindings(TFunctionRef<EVisitResult(FPropertyBindingBinding& Binding)> InFunction) override;
+	SCENESTATEBINDING_API virtual void OnReset() override;
+	SCENESTATEBINDING_API virtual void VisitSourceStructDescriptorInternal(TFunctionRef<EVisitResult(const FPropertyBindingBindableStructDescriptor&)> InFunction) const override;
+	[[nodiscard]] SCENESTATEBINDING_API virtual bool OnResolvingPaths() override;
+	//~ End FPropertyBindingBindingCollection
+
+private:
+	UPROPERTY()
+	TArray<FSceneStateBindingDesc> BindingDescs;
+
+	/** Contains all the bindings from a source data to a target data */
+	UPROPERTY()
+	TArray<FSceneStateBinding> Bindings;
+
+	/**
+	 * Contains all the reference bindings from a source data.
+	 * This gets set and populated during scene state compilation.
+	 */
+	UPROPERTY()
+	TArray<FSceneStateBindingReference> References;
+
+	/** Contains the references that have been resolved for use at runtime */
+	UPROPERTY()
+	TArray<FSceneStateBindingResolvedReference> ResolvedReferences;
+
+	friend UE::SceneState::Editor::FBindingCompiler;
+};
diff --git a/Engine/Plugins/VirtualProduction/SceneState/Source/SceneStateBinding/Public/SceneStateBindingUtils.h b/Engine/Plugins/VirtualProduction/SceneState/Source/SceneStateBinding/Public/SceneStateBindingUtils.h
new file mode 100644
--- /dev/null
+++ b/Engine/Plugins/VirtualProduction/SceneState/Source/SceneStateBinding/Public/SceneStateBindingUtils.h
@@ -0,0 +1,60 @@
+// Copyright Epic Games, Inc. All Rights Reserved.
+
+#pragma once
+
+#include "PropertyBindingDataView.h"
+#include "Templates/Function.h"
+
+class UObject;
+struct FGuid;
+struct FSceneStateBindingCollection;
+struct FSceneStateBindingDataHandle;
+
+namespace UE::SceneState
+{
+
+#if WITH_EDITOR
+	/**
+	 * Update the bindings of a source struct id into a target one.
+	 * It goes through the given object's outer chain to get to the binding collection owner to do this operation
+	 * @param InObject object to use to find the binding collection owner
+	 * @param InOldStructId the struct id to queue copy bindings from
+	 * @param InNewStructId the struct id to queue copy bindings to
+	 */
+	SCENESTATEBINDING_API void HandleStructIdChanged(UObject& InObject, const FGuid& InOldStructId, const FGuid& InNewStructId);
+#endif
+
+	struct FApplyBatchParams
+	{
+		/** Batch within the Binding Collection to run */
+		uint16 BatchIndex;
+		/** Target Data to copy the source data to */
+		FPropertyBindingDataView TargetDataView;
+		/** Binding collection holding the Copy Batch to run */
+		const FSceneStateBindingCollection& BindingCollection;
+		/** Functor to find the Source Data View for a given Data Handle */
+		TFunctionRef<FPropertyBindingDataView(const FSceneStateBindingDataHandle&)> FindSourceDataViewFunctor;
+	};
+	/** Runs the Batch within the Binding Collection for a given Target Data */
+	SCENESTATEBINDING_API bool ApplyBatch(const FApplyBatchParams& InParams);
+
+	struct FPatchBindingParams
+	{
+		/** Binding collection to patch */
+		FSceneStateBindingCollection& BindingCollection;
+		/** Functor to find the data struct for a given Data Handle */
+		TFunctionRef<const UStruct*(const FSceneStateBindingDataHandle&)> FindDataStructFunctor;
+	};
+	/** Patches invalidated structs (property bags, user defined structs, etc.) in the given binding collection */
+	SCENESTATEBINDING_API void PatchBindingCollection(const FPatchBindingParams& InParams);
+
+	/** Patches any invalidated structs (property bags, user defined structs, etc.) in the binding descs of the given binding collection */
+	void PatchBindingDescs(const FPatchBindingParams& InParams);
+
+	/** Patches any invalidated structs (property bags, user defined structs, etc.) in the bindings of the given binding collection */
+	void PatchBindings(const FPatchBindingParams& InParams);
+
+	/** Patches any invalidated structs (property bags, user defined structs, etc.) in the copy batches of the given binding collection */
+	void PatchCopyBatches(const FPatchBindingParams& InParams);
+
+} // UE::SceneState
diff --git a/Engine/Plugins/VirtualProduction/SceneState/Source/SceneStateBlueprint/Public/SceneStateBlueprint.h b/Engine/Plugins/VirtualProduction/SceneState/Source/SceneStateBlueprint/Public/SceneStateBlueprint.h
new file mode 100644
--- /dev/null
+++ b/Engine/Plugins/VirtualProduction/SceneState/Source/SceneStateBlueprint/Public/SceneStateBlueprint.h
@@ -0,0 +1,91 @@
+// Copyright Epic Games, Inc. All Rights Reserved.
+
+#pragma once
+
+#include "Containers/Array.h"
+#include "Engine/Blueprint.h"
+#include "PropertyBindingBindingCollectionOwner.h"
+#include "SceneStateBindingCollection.h"
+#include "SceneStateBlueprint.generated.h"
+
+class USceneStateGeneratedClass;
+class USceneStateMachineGraph;
+class USceneStateMachineStateNode;
+
+UCLASS(MinimalAPI, DisplayName="Motion Design Scene State Blueprint")
+class USceneStateBlueprint : public UBlueprint, public IPropertyBindingBindingCollectionOwner
+{
+	GENERATED_BODY()
+
+public:
+	SCENESTATEBLUEPRINT_API USceneStateBlueprint(const FObjectInitializer& InObjectInitializer);
+
+	const FGuid& GetRootId() const
+	{
+		return RootId;
+	}
+
+	template<typename T>
+	T* FindExtension() const
+	{
+		return Cast<T>(FindExtension(T::StaticClass()));
+	}
+
+	SCENESTATEBLUEPRINT_API UBlueprintExtension* FindExtension(TSubclassOf<UBlueprintExtension> InClass) const;
+
+	SCENESTATEBLUEPRINT_API FSceneStateBindingDesc CreateRootBinding() const;
+
+	//~ Begin IPropertyBindingBindingCollectionOwner
+	SCENESTATEBLUEPRINT_API virtual void GetBindableStructs(const FGuid InTargetStructId, TArray<TInstancedStruct<FPropertyBindingBindableStructDescriptor>>& OutStructDescs) const override;
+	SCENESTATEBLUEPRINT_API virtual bool GetBindableStructByID(const FGuid InStructId, TInstancedStruct<FPropertyBindingBindableStructDescriptor>& OutStructDesc) const override;
+	SCENESTATEBLUEPRINT_API virtual bool GetBindingDataViewByID(const FGuid InStructId, FPropertyBindingDataView& OutDataView) const override;
+	SCENESTATEBLUEPRINT_API virtual FPropertyBindingBindingCollection* GetEditorPropertyBindings() override;
+	SCENESTATEBLUEPRINT_API virtual const FPropertyBindingBindingCollection* GetEditorPropertyBindings() const override;
+	SCENESTATEBLUEPRINT_API virtual bool CanCreateParameter(const FGuid InStructId) const override;
+	SCENESTATEBLUEPRINT_API virtual void CreateParametersForStruct(const FGuid InStructId, TArrayView<UE::PropertyBinding::FPropertyCreationDescriptor> InOutCreationDescs) override;
+	//~ End IPropertyBindingBindingCollectionOwner
+
+	//~ Begin UBlueprint
+	SCENESTATEBLUEPRINT_API virtual void SetObjectBeingDebugged(UObject* InNewObject) override;
+	SCENESTATEBLUEPRINT_API virtual UClass* GetBlueprintClass() const override;
+	SCENESTATEBLUEPRINT_API virtual void GetReparentingRules(TSet<const UClass*>& OutAllowedChildrenOfClasses, TSet<const UClass*>& OutDisallowedChildrenOfClasses) const override;
+	SCENESTATEBLUEPRINT_API virtual bool SupportedByDefaultBlueprintFactory() const override;
+	SCENESTATEBLUEPRINT_API virtual void LoadModulesRequiredForCompilation() override;
+	SCENESTATEBLUEPRINT_API virtual bool IsValidForBytecodeOnlyRecompile() const override;
+	//~ End UBlueprint
+
+	//~ Begin UObject
+	SCENESTATEBLUEPRINT_API virtual void BeginDestroy() override;
+	//~ End UObject
+
+	/** The top level State Machine Graphs of this State. Does not include the state machine graphs of nested state nodes */
+	UPROPERTY()
+	TArray<TObjectPtr<UEdGraph>> StateMachineGraphs;
+
+	/** Holds all the editor bindings prior to compilation */
+	UPROPERTY()
+	FSceneStateBindingCollection BindingCollection;
+
+private:
+	/** Called when a blueprint variable has been renamed. Used to fix the bindings to the variable being renamed to use the new name */
+	void OnRenameVariableReferences(UBlueprint* InBlueprint, UClass* InVariableClass, const FName& InOldVariableName, const FName& InNewVariableName);
+
+	/** Called when a blueprint variable has been renamed. Used to fix the path to the variable being renamed to use the new name */
+	void RenameVariableReferenceInPath(FPropertyBindingPath& InPath, FName InOldVariableName, FName InNewVariableName);
+
+	/** Called when a state machine graph's parameters have changed */
+	void OnGraphParametersChanged(USceneStateMachineGraph* InGraph);
+
+	/** Called when a state machine graph's parameters have changed. Used to fix the path to parameters that have possibly been renamed */
+	void UpdateGraphParametersBindings(FPropertyBindingPath& InPath, USceneStateMachineGraph* InGraph);
+
+	/** Unique id representing this blueprint / generated class as a bindable struct */
+	UPROPERTY()
+	FGuid RootId;
+
+	/** Handle to the delegate when a blueprint variable has been renamed */
+	FDelegateHandle OnRenameVariableReferencesHandle;
+
+	/** Handle to the delegate when a state machine graph parameters have changed */
+	FDelegateHandle OnGraphParametersChangedHandle;
+};
diff --git a/Engine/Plugins/VirtualProduction/SceneState/Source/SceneStateBlueprintEditor/Private/DebugControls/SceneStateDebugControlsObjectDetails.cpp b/Engine/Plugins/VirtualProduction/SceneState/Source/SceneStateBlueprintEditor/Private/DebugControls/SceneStateDebugControlsObjectDetails.cpp
new file mode 100644
--- /dev/null
+++ b/Engine/Plugins/VirtualProduction/SceneState/Source/SceneStateBlueprintEditor/Private/DebugControls/SceneStateDebugControlsObjectDetails.cpp
@@ -0,0 +1,84 @@
+// Copyright Epic Games, Inc. All Rights Reserved.
+
+#include "SceneStateDebugControlsObjectDetails.h"
+#include "DetailCategoryBuilder.h"
+#include "DetailLayoutBuilder.h"
+#include "DetailWidgetRow.h"
+#include "IDetailChildrenBuilder.h"
+#include "PropertyCustomizationHelpers.h"
+#include "SceneStateDebugControlsObject.h"
+#include "SceneStateEventUtils.h"
+#include "SceneStateObject.h"
+#include "Widgets/Input/SButton.h"
+#include "Widgets/SBoxPanel.h"
+#include "Widgets/Text/STextBlock.h"
+
+#define LOCTEXT_NAMESPACE "SceneStateDebugControlsObjectDetails"
+
+namespace UE::SceneState::Editor
+{
+
+void FDebugControlsObjectDetails::CustomizeDetails(IDetailLayoutBuilder& InDetailBuilder)
+{
+	DebugControls = InDetailBuilder.GetObjectsOfTypeBeingCustomized<USceneStateDebugControlsObject>();
+	CustomizeEventDetails(InDetailBuilder);
+	CustomizeDebuggedObjectDetails(InDetailBuilder);
+}
+
+void FDebugControlsObjectDetails::CustomizeEventDetails(IDetailLayoutBuilder& InDetailBuilder)
+{
+	EventsHandle = InDetailBuilder.GetProperty(GET_MEMBER_NAME_CHECKED(USceneStateDebugControlsObject, Events));
+	EventsHandle->MarkHiddenByCustomization();
+
+	IDetailCategoryBuilder& EventsCategory = InDetailBuilder.EditCategory(TEXT("Events"));
+	EventsCategory.AddProperty(EventsHandle);
+}
+
+USceneStateObject* FDebugControlsObjectDetails::GetDebuggedObject(int32 InIndex) const
+{
+	if (!DebugControls.IsValidIndex(InIndex))
+	{
+		return nullptr;
+	}
+
+	USceneStateDebugControlsObject* const DebugControlsObject = DebugControls[InIndex].Get();
+	if (!DebugControlsObject)
+	{
+		return nullptr;
+	}
+
+	return DebugControlsObject->DebuggedObjectWeak.Get();
+}
+
+void FDebugControlsObjectDetails::CustomizeDebuggedObjectDetails(IDetailLayoutBuilder& InDetailBuilder)
+{
+	TArray<UObject*> DebuggedObjects;
+	DebuggedObjects.Reserve(DebugControls.Num());
+
+	// Gather all the Debugged objects within the debug controls
+	Algo::Transform(DebugControls, DebuggedObjects,
+		[](const TWeakObjectPtr<USceneStateDebugControlsObject>& InDebugControls)->UObject*
+		{
+			if (USceneStateDebugControlsObject* DebugControl = InDebugControls.Get())
+			{
+				return DebugControl->DebuggedObjectWeak.Get();
+			}
+			return nullptr;
+		});
+
+	DebuggedObjects.RemoveAll(
+		[](UObject* InObject)
+		{
+			return !InObject;
+		});
+
+	InDetailBuilder
+		.EditCategory(TEXT("Debugged Object"))
+		.AddExternalObjects(DebuggedObjects, EPropertyLocation::Default,
+			FAddPropertyParams()
+			.HideRootObjectNode(/*bHideRootObjectNode*/true));
+}
+
+} // UE::SceneState::Editor
+
+#undef LOCTEXT_NAMESPACE
diff --git a/Engine/Plugins/VirtualProduction/SceneState/Source/SceneStateBlueprintEditor/Private/DebugExecutors/SceneStateDebugTaskExecutor.cpp b/Engine/Plugins/VirtualProduction/SceneState/Source/SceneStateBlueprintEditor/Private/DebugExecutors/SceneStateDebugTaskExecutor.cpp
new file mode 100644
--- /dev/null
+++ b/Engine/Plugins/VirtualProduction/SceneState/Source/SceneStateBlueprintEditor/Private/DebugExecutors/SceneStateDebugTaskExecutor.cpp
@@ -0,0 +1,105 @@
+// Copyright Epic Games, Inc. All Rights Reserved.
+
+#include "SceneStateDebugTaskExecutor.h"
+#include "Nodes/SceneStateMachineStateNode.h"
+#include "Nodes/SceneStateMachineTaskNode.h"
+#include "SceneStateGeneratedClass.h"
+#include "SceneStateMachineGraphSchema.h"
+#include "Tasks/SceneStateTaskInstance.h"
+
+namespace UE::SceneState::Editor
+{
+
+const FSceneStateTask* FDebugTaskExecutor::GetTask(const FSceneStateExecutionContext& InExecutionContext) const
+{
+	if (const USceneStateGeneratedClass* GeneratedClass = InExecutionContext.GetGeneratedClass())
+	{
+		return GeneratedClass->FindTaskFromNode(GetNodeKey());
+	}
+	return nullptr;
+}
+
+const FSceneState* FDebugTaskExecutor::GetState(const USceneStateGeneratedClass* InGeneratedClass) const
+{
+	if (const USceneStateMachineTaskNode* const TaskNode = Cast<USceneStateMachineTaskNode>(GetNodeKey().ResolveObjectPtr()))
+	{
+		if (const USceneStateMachineStateNode* const StateNode = USceneStateMachineGraphSchema::FindConnectedStateNode(TaskNode))
+		{
+			return InGeneratedClass->FindStateFromNode(StateNode);
+		}
+	}
+	return nullptr;
+}
+
+void FDebugTaskExecutor::Setup(const FSceneStateExecutionContext& InExecutionContext, const FSceneStateTask& InTask)
+{
+	const USceneStateGeneratedClass* GeneratedClass = InExecutionContext.GetGeneratedClass();
+	if (!GeneratedClass)
+	{
+		return;
+	}
+
+	const FSceneState* ParentState = GetState(GeneratedClass);
+	if (!ParentState)
+	{
+		return;
+	}
+
+	const FSceneStateRange TaskRange = ParentState->GetTaskRange(); 
+
+	const uint16 TaskAbsoluteIndex = InTask.GetTaskIndex();
+	const uint16 TaskRelativeIndex = TaskAbsoluteIndex - TaskRange.Index;
+
+	// Make an empty task instance array that matches the state task count,
+	// and only set the task instance template for the task this executor cares about
+	TArray<FConstStructView> TemplateTaskInstances;
+	TemplateTaskInstances.SetNum(TaskRange.Count);
+	TemplateTaskInstances[TaskRelativeIndex] = GeneratedClass->GetTaskInstances()[TaskAbsoluteIndex];
+
+	ParentState->AllocateTaskInstances(InExecutionContext, TemplateTaskInstances);
+}
+
+void FDebugTaskExecutor::OnStart(const FSceneStateExecutionContext& InExecutionContext)
+{
+	if (const FSceneStateTask* Task = GetTask(InExecutionContext))
+	{
+		Setup(InExecutionContext, *Task);
+
+		const FStructView TaskInstance = Task->FindTaskInstance(InExecutionContext);
+
+		Task->Setup(InExecutionContext, TaskInstance);
+		Task->Start(InExecutionContext, TaskInstance);
+		ConditionallyExit(TaskInstance);
+	}
+}
+
+void FDebugTaskExecutor::OnTick(const FSceneStateExecutionContext& InExecutionContext, float InDeltaSeconds)
+{
+	if (const FSceneStateTask* Task = GetTask(InExecutionContext))
+	{
+		const FStructView TaskInstance = Task->FindTaskInstance(InExecutionContext);
+
+		Task->Tick(InExecutionContext, TaskInstance, InDeltaSeconds);
+		ConditionallyExit(TaskInstance);
+	}
+}
+
+void FDebugTaskExecutor::OnExit(const FSceneStateExecutionContext& InExecutionContext)
+{
+	if (const FSceneStateTask* Task = GetTask(InExecutionContext))
+	{
+		const FStructView TaskInstance = Task->FindTaskInstance(InExecutionContext);
+		Task->Stop(InExecutionContext, TaskInstance, ESceneStateTaskStopReason::State);
+	}
+}
+
+void FDebugTaskExecutor::ConditionallyExit(FStructView InTaskInstance)
+{
+	const FSceneStateTaskInstance* TaskInstance = InTaskInstance.GetPtr<FSceneStateTaskInstance>();
+	if (TaskInstance && TaskInstance->GetStatus() == EExecutionStatus::Finished)
+	{
+		Exit();
+	}
+}
+
+} // UE::SceneState::Editor
diff --git a/Engine/Plugins/VirtualProduction/SceneState/Source/SceneStateBlueprintEditor/Public/SceneStateBlueprintEditorUtils.h b/Engine/Plugins/VirtualProduction/SceneState/Source/SceneStateBlueprintEditor/Public/SceneStateBlueprintEditorUtils.h
new file mode 100644
--- /dev/null
+++ b/Engine/Plugins/VirtualProduction/SceneState/Source/SceneStateBlueprintEditor/Public/SceneStateBlueprintEditorUtils.h
@@ -0,0 +1,72 @@
+// Copyright Epic Games, Inc. All Rights Reserved.
+
+#pragma once
+
+#include "PropertyEditorModule.h"
+#include "Templates/SharedPointerFwd.h"
+
+class IDetailChildrenBuilder;
+class IPropertyHandle;
+class UClass;
+struct FGuid;
+struct FInstancedPropertyBag;
+
+namespace UE::SceneState::Editor
+{
+	/** Finds the Task Id for a given Property Handle by first finding the outer Task Node and returning its Task id. */
+	SCENESTATEBLUEPRINTEDITOR_API FGuid FindTaskId(const TSharedRef<IPropertyHandle>& InPropertyHandle);
+
+	/**
+	 * Returns whether a given property is an object property of a given class
+	 * @param InProperty the property to check
+	 * @param InClass the class to check
+	 * @return true if the property is an object property of the given class
+	 */
+	SCENESTATEBLUEPRINTEDITOR_API bool IsObjectPropertyOfClass(const FProperty* InProperty, const UClass* InClass);
+
+	/**
+	 * Returns whether a given property is of a given struct type
+	 * @param InProperty the property to check
+	 * @param InStruct the type to check
+	 * @return true if the property is a struct property of the given struct type  
+	 */
+	SCENESTATEBLUEPRINTEDITOR_API bool IsStruct(const FProperty* InProperty, const UScriptStruct* InStruct);
+
+	/**
+	 * Gets the Guid value from the Property Handle.
+	 * @param InGuidPropertyHandle the guid property handle to get the guid value from
+	 * @param OutGuid the returned guid if the operation was successful
+	 * @return whether the guid was obtained successfully, or whether it failed or had multiple values
+	 */
+	SCENESTATEBLUEPRINTEDITOR_API FPropertyAccess::Result GetGuid(const TSharedRef<IPropertyHandle>& InGuidPropertyHandle, FGuid& OutGuid);
+
+	/**
+	 * Tries to create a structure data provider for the given instanced struct property handle
+	 * @param InStructHandle the instanced struct property handle to create a struct data provider for
+	 * @return the created struct data provider
+	 */
+	SCENESTATEBLUEPRINTEDITOR_API TSharedRef<IStructureDataProvider> CreateInstancedStructDataProvider(const TSharedRef<IPropertyHandle>& InStructHandle);
+
+	/**
+	 * Compares the structure layout of two property bags
+	 * @param InParametersA the first property bag comparand
+	 * @param InParametersB the second property bag comparand
+	 * @return true if both have the same number of properties with same names and compatible types.
+	 */
+	SCENESTATEBLUEPRINTEDITOR_API bool CompareParametersLayout(const FInstancedPropertyBag& InParametersA, const FInstancedPropertyBag& InParametersB);
+
+	/** Sets the Instance Meta-data of a given property handle to a Task Id guid */
+	SCENESTATEBLUEPRINTEDITOR_API void AssignBindingId(const TSharedRef<IPropertyHandle>& InPropertyHandle, const FGuid& InTaskId);
+
+	/** Finds the Common Class for all the Objects within the Object Property Handle */
+	SCENESTATEBLUEPRINTEDITOR_API UClass* FindCommonBase(const TSharedRef<IPropertyHandle>& InPropertyHandle);
+
+	/**
+	 * Helper function to add the Common properties of the object property handle to the Child Builder
+	 * @param InPropertyHandle the object property handle to get the common properties from
+	 * @param InChildBuilder the builder to add the properties to
+	 * @param InDisallowedFlags filter flags. Any property with any of these flags will be skipped 
+	 */
+	SCENESTATEBLUEPRINTEDITOR_API bool AddObjectProperties(const TSharedRef<IPropertyHandle>& InPropertyHandle, IDetailChildrenBuilder& InChildBuilder, uint64 InDisallowedFlags = CPF_DisableEditOnInstance);
+
+} // UE::SceneState::Editor
diff --git a/Engine/Plugins/VirtualProduction/SceneState/Source/SceneStateEvent/Public/SceneStateEventStream.h b/Engine/Plugins/VirtualProduction/SceneState/Source/SceneStateEvent/Public/SceneStateEventStream.h
new file mode 100644
--- /dev/null
+++ b/Engine/Plugins/VirtualProduction/SceneState/Source/SceneStateEvent/Public/SceneStateEventStream.h
@@ -0,0 +1,59 @@
+// Copyright Epic Games, Inc. All Rights Reserved.
+
+#pragma once
+
+#include "Containers/Array.h"
+#include "Containers/Map.h"
+#include "SceneStateEvent.h"
+#include "UObject/Object.h"
+#include "SceneStateEventStream.generated.h"
+
+struct FSceneStateEventHandler;
+struct FSceneStateEventSchemaHandle;
+
+/** Holds and keeps track of Events added at Runtime */
+UCLASS(MinimalAPI)
+class USceneStateEventStream : public UObject
+{
+	GENERATED_BODY()
+
+public:
+	/** Registers the Event Stream to the Event Subsystem to listen to Broadcast Events */
+	SCENESTATEEVENT_API bool Register();
+
+	/** Registers the Event Stream from the Event Subsystem */
+	SCENESTATEEVENT_API void Unregister();
+
+	/** Pushes a new event into the stream */
+	SCENESTATEEVENT_API void PushEvent(FSharedStruct&& InEvent);
+
+	/** Pushes a new event into the stream */
+	SCENESTATEEVENT_API void PushEvent(const FSharedStruct& InEvent);
+
+	/** Consumes the first (oldest) Event that was pushed that matches the given Schema */
+	SCENESTATEEVENT_API bool ConsumeEventBySchema(const FSceneStateEventSchemaHandle& InEventSchemaHandle);
+
+	/** Finds the first (oldest) Event that was pushed that matches the given Schema */
+	SCENESTATEEVENT_API const FSceneStateEvent* FindEventBySchema(const FSceneStateEventSchemaHandle& InEventSchemaHandle) const;
+
+	/** Finds the event that is captured by the Handler with the given Id */
+	SCENESTATEEVENT_API FSceneStateEvent* FindCapturedEvent(const FGuid& InHandlerId);
+
+	/** Removes the Events (first/oldest) that match the given Event Handlers and moves them to the Captured Event map */
+	SCENESTATEEVENT_API void CaptureEvents(TConstArrayView<FSceneStateEventHandler> InEventHandlers);
+
+	/** Cleanup the Captured Events that match the given Event Handlers */
+	SCENESTATEEVENT_API void ResetCapturedEvents(TConstArrayView<FSceneStateEventHandler> InEventHandlers);
+
+private:
+	/** Returns the index of the first event matching the event schema */
+	int32 GetEventIndexBySchema(const FSceneStateEventSchemaHandle& InEventSchemaHandle) const;
+
+	/** Active Events kept in push order */
+	UPROPERTY(Transient, DuplicateTransient, TextExportTransient)
+	TArray<FSharedStruct> Events;
+
+	/** Map of the Handler Id to the Event it has captured */
+	UPROPERTY(Transient, DuplicateTransient, TextExportTransient)
+	TMap<FGuid, FSharedStruct> CapturedEvents;
+};
diff --git a/Engine/Plugins/VirtualProduction/SceneState/Source/SceneStateGameplay/Public/SceneStateComponentPlayer.h b/Engine/Plugins/VirtualProduction/SceneState/Source/SceneStateGameplay/Public/SceneStateComponentPlayer.h
new file mode 100644
--- /dev/null
+++ b/Engine/Plugins/VirtualProduction/SceneState/Source/SceneStateGameplay/Public/SceneStateComponentPlayer.h
@@ -0,0 +1,23 @@
+// Copyright Epic Games, Inc. All Rights Reserved.
+
+#pragma once
+
+#include "SceneStatePlayer.h"
+#include "SceneStateComponentPlayer.generated.h"
+
+class USceneStateComponent;
+
+/** Scene State Player for Scene State Components */
+UCLASS(MinimalAPI)
+class USceneStateComponentPlayer : public USceneStatePlayer
+{
+	GENERATED_BODY()
+
+protected:
+	SCENESTATEGAMEPLAY_API AActor* GetActor() const;
+
+	//~ Begin USceneStatePlayer
+	SCENESTATEGAMEPLAY_API virtual bool OnGetContextName(FString& OutContextName) const override;
+	SCENESTATEGAMEPLAY_API virtual bool OnGetContextObject(UObject*& OutContextObject) const override;
+	//~ End USceneStatePlayer
+};
diff --git a/Engine/Plugins/VirtualProduction/SceneState/Source/SceneStateMachineGraph/Public/SceneStateMachineGraphUtils.h b/Engine/Plugins/VirtualProduction/SceneState/Source/SceneStateMachineGraph/Public/SceneStateMachineGraphUtils.h
new file mode 100644
--- /dev/null
+++ b/Engine/Plugins/VirtualProduction/SceneState/Source/SceneStateMachineGraph/Public/SceneStateMachineGraphUtils.h
@@ -0,0 +1,15 @@
+// Copyright Epic Games, Inc. All Rights Reserved.
+
+#pragma once
+
+#include "HAL/Platform.h"
+
+class UEdGraph;
+
+namespace UE::SceneState::Graph
+{
+	SCENESTATEMACHINEGRAPH_API bool CanDirectlyRemoveGraph(UEdGraph* InGraph);
+
+	SCENESTATEMACHINEGRAPH_API void RemoveGraph(UEdGraph* InGraph);
+
+} // UE::SceneState::Graph
