Motion Design Scene State: moved apply batch to scene state module. This is to simplify the function itself by not requiring individual components and just using the execution context to get the information. The usage of this function is also only used within the scene state module.

diff --git a/Engine/Plugins/VirtualProduction/SceneState/Source/SceneState/Private/SceneStateMachine.cpp b/Engine/Plugins/VirtualProduction/SceneState/Source/SceneState/Private/SceneStateMachine.cpp
--- a/Engine/Plugins/VirtualProduction/SceneState/Source/SceneState/Private/SceneStateMachine.cpp
+++ b/Engine/Plugins/VirtualProduction/SceneState/Source/SceneState/Private/SceneStateMachine.cpp
@@ -3,12 +3,12 @@
 #include "SceneStateMachine.h"
 #include "Conduit/SceneStateConduit.h"
 #include "SceneState.h"
-#include "SceneStateBindingUtils.h"
 #include "SceneStateExecutionContext.h"
 #include "SceneStateExecutionScope.h"
 #include "SceneStateLog.h"
 #include "SceneStateMachineInstance.h"
 #include "SceneStateObject.h"
+#include "SceneStateUtils.h"
 #include "Transition/SceneStateTransition.h"
 #include "Transition/SceneStateTransitionEvaluation.h"
 
@@ -229,27 +229,13 @@ bool FSceneStateMachine::ApplyBindings(const FSceneStateExecutionContext& InCont
 {
 	QUICK_SCOPE_CYCLE_COUNTER(SceneStateMachine_ApplyBindings);
 
-	const FSceneStateBindingCollection* const BindingCollection = InContext.GetBindingCollection();
-	if (!BindingCollection)
-	{
-		return false;
-	}
-
-	auto FindSourceDataView =
-		[&InContext](const FSceneStateBindingDataHandle& InDataHandle)
-		{
-			return InContext.FindDataView(InDataHandle);
-		};
-
-	UE::SceneState::FApplyBatchParams ApplyBatchParams
+	const UE::SceneState::FApplyBatchParams ApplyBatchParams
 		{
-			.BatchIndex = BindingsBatch.Get(),
+			.BindingsBatch = BindingsBatch.Get(),
 			.TargetDataView = InStateMachineInstance.Parameters.GetMutableValue(),
-			.BindingCollection = *BindingCollection,
-			.FindSourceDataViewFunctor = FindSourceDataView,
 		};
 
-	return ApplyBatch(ApplyBatchParams);
+	return ApplyBatch(InContext, ApplyBatchParams);
 }
 
 void FSceneStateMachine::StopActiveState(const FSceneStateExecutionContext& InContext, FSceneStateMachineInstance& InInstance) const
diff --git a/Engine/Plugins/VirtualProduction/SceneState/Source/SceneState/Private/SceneStateUtils.cpp b/Engine/Plugins/VirtualProduction/SceneState/Source/SceneState/Private/SceneStateUtils.cpp
--- a/Engine/Plugins/VirtualProduction/SceneState/Source/SceneState/Private/SceneStateUtils.cpp
+++ b/Engine/Plugins/VirtualProduction/SceneState/Source/SceneState/Private/SceneStateUtils.cpp
@@ -2,6 +2,10 @@
 
 #include "SceneStateUtils.h"
 #include "Engine/Engine.h"
+#include "PropertyBindingPath.h"
+#include "PropertyBindingTypes.h"
+#include "SceneStateBindingCollection.h"
+#include "SceneStateExecutionContext.h"
 #include "SceneStateLog.h"
 #include "SceneStateRange.h"
 #include "StructUtils/InstancedStructContainer.h"
@@ -63,6 +67,34 @@ TArray<FConstStructView> GetConstStructViews(const FInstancedStructContainer& In
 	return GetConstStructViews(InStructContainer, Range);
 }
 
+bool ApplyBatch(const FSceneStateExecutionContext& InContext, const FApplyBatchParams& InParams)
+{
+	const FPropertyBindingIndex16 BindingsBatch(InParams.BindingsBatch);
+	if (!BindingsBatch.IsValid())
+	{
+		return false;
+	}
+
+	const FSceneStateBindingCollection* const BindingCollection = InContext.GetBindingCollection();
+	if (!BindingCollection)
+	{
+		return false;
+	}
+
+	bool bResult = true;
+
+	const FPropertyBindingCopyInfoBatch& Batch = BindingCollection->GetBatch(BindingsBatch);
+	check(InParams.TargetDataView.GetStruct() == Batch.TargetStruct.Get().Struct);
+
+	for (const FPropertyBindingCopyInfo& Copy : BindingCollection->GetBatchCopies(Batch))
+	{
+		const FPropertyBindingDataView SourceView = InContext.FindDataView(Copy.SourceDataHandle.Get<FSceneStateBindingDataHandle>());
+		bResult &= BindingCollection->CopyProperty(Copy, SourceView, InParams.TargetDataView);
+	}
+
+	return bResult;
+}
+
 void* ResolveVisitedPath(const UScriptStruct* InRootObject, void* InRootData, const FPropertyVisitorPath& InPath)
 {
 	if (!InRootObject || !InRootData)
diff --git a/Engine/Plugins/VirtualProduction/SceneState/Source/SceneState/Private/Tasks/SceneStateTask.cpp b/Engine/Plugins/VirtualProduction/SceneState/Source/SceneState/Private/Tasks/SceneStateTask.cpp
--- a/Engine/Plugins/VirtualProduction/SceneState/Source/SceneState/Private/Tasks/SceneStateTask.cpp
+++ b/Engine/Plugins/VirtualProduction/SceneState/Source/SceneState/Private/Tasks/SceneStateTask.cpp
@@ -1,7 +1,6 @@
 // Copyright Epic Games, Inc. All Rights Reserved.
 
 #include "Tasks/SceneStateTask.h"
-#include "SceneStateBindingUtils.h"
 #include "SceneStateExecutionContext.h"
 #include "SceneStateObject.h"
 #include "SceneStateUtils.h"
@@ -160,36 +159,22 @@ bool FSceneStateTask::ApplyBindings(const FSceneStateExecutionContext& InContext
 {
 	QUICK_SCOPE_CYCLE_COUNTER(SceneStateTask_ApplyBindings);
 
-	const FSceneStateBindingCollection* const BindingCollection = InContext.GetBindingCollection();
-	if (!BindingCollection)
-	{
-		return false;
-	}
-
-	auto FindDataView =
-		[&InContext](const FSceneStateBindingDataHandle& InDataHandle)
-		{
-			return InContext.FindDataView(InDataHandle);
-		};
-
 	UE::SceneState::FApplyBatchParams ApplyBatchParams
 		{
-			.BatchIndex = BindingsBatch.Get(),
-			.TargetDataView = InTaskInstance,
-			.BindingCollection = *BindingCollection,
-			.FindSourceDataViewFunctor = FindDataView,
+			.BindingsBatch = BindingsBatch.Get(),
+			.TargetDataView = InTaskInstance
 		};
 
-	bool bResult = ApplyBatch(ApplyBatchParams);
+	bool bResult = ApplyBatch(InContext, ApplyBatchParams);
 
 	if (const FSceneStateTaskBindingExtension* BindingExtension = GetBindingExtension())
 	{
 		BindingExtension->VisitBindingBatches(InTaskInstance,
-			[&bResult, &ApplyBatchParams](uint16 InBatchIndex, FStructView InTargetDataView)
+			[&bResult, &InContext, &ApplyBatchParams](uint16 InBatchIndex, FStructView InTargetDataView)
 			{
-				ApplyBatchParams.BatchIndex = InBatchIndex;
+				ApplyBatchParams.BindingsBatch = InBatchIndex;
 				ApplyBatchParams.TargetDataView = InTargetDataView;
-				bResult &= ApplyBatch(ApplyBatchParams);
+				bResult &= ApplyBatch(InContext, ApplyBatchParams);
 			});
 	}
 
diff --git a/Engine/Plugins/VirtualProduction/SceneState/Source/SceneState/Private/Transition/SceneStateTransition.cpp b/Engine/Plugins/VirtualProduction/SceneState/Source/SceneState/Private/Transition/SceneStateTransition.cpp
--- a/Engine/Plugins/VirtualProduction/SceneState/Source/SceneState/Private/Transition/SceneStateTransition.cpp
+++ b/Engine/Plugins/VirtualProduction/SceneState/Source/SceneState/Private/Transition/SceneStateTransition.cpp
@@ -1,10 +1,10 @@
 // Copyright Epic Games, Inc. All Rights Reserved.
 
 #include "Transition/SceneStateTransition.h"
-#include "SceneStateBindingUtils.h"
 #include "SceneStateExecutionContext.h"
 #include "SceneStateGeneratedClass.h"
 #include "SceneStateObject.h"
+#include "SceneStateUtils.h"
 #include "Transition/SceneStateTransitionEvaluation.h"
 #include "Transition/SceneStateTransitionInstance.h"
 #include "Transition/SceneStateTransitionLink.h"
@@ -83,27 +83,13 @@ bool FSceneStateTransition::ApplyBindings(const FSceneStateExecutionContext& InC
 {
 	QUICK_SCOPE_CYCLE_COUNTER(SceneStateTransition_ApplyBindings);
 
-	const FSceneStateBindingCollection* const BindingCollection = InContext.GetBindingCollection();
-	if (!BindingCollection)
-	{
-		return false;
-	}
-
-	auto FindDataView =
-		[&InContext](const FSceneStateBindingDataHandle& InDataHandle)
-		{
-			return InContext.FindDataView(InDataHandle);
-		};
-
-	UE::SceneState::FApplyBatchParams ApplyBatchParams
+	const UE::SceneState::FApplyBatchParams ApplyBatchParams
 		{
-			.BatchIndex = BindingsBatch.Get(),
+			.BindingsBatch = BindingsBatch.Get(),
 			.TargetDataView = InInstance.Parameters.GetMutableValue(),
-			.BindingCollection = *BindingCollection,
-			.FindSourceDataViewFunctor = FindDataView,
 		};
 
-	return ApplyBatch(ApplyBatchParams);
+	return ApplyBatch(InContext, ApplyBatchParams);
 }
 
 bool FSceneStateTransition::ContainsAllRequiredEvents(const UE::SceneState::FTransitionEvaluationParams& InParams) const
diff --git a/Engine/Plugins/VirtualProduction/SceneState/Source/SceneState/Public/SceneStateUtils.h b/Engine/Plugins/VirtualProduction/SceneState/Source/SceneState/Public/SceneStateUtils.h
--- a/Engine/Plugins/VirtualProduction/SceneState/Source/SceneState/Public/SceneStateUtils.h
+++ b/Engine/Plugins/VirtualProduction/SceneState/Source/SceneState/Public/SceneStateUtils.h
@@ -4,13 +4,17 @@
 
 #include "Containers/Map.h"
 #include "HAL/Platform.h"
+#include "PropertyBindingDataView.h"
 #include "Templates/Function.h"
 
+#define UE_API SCENESTATE_API
+
 class UClass;
 class UObject;
 struct FConstStructView;
 struct FInstancedStructContainer;
 struct FPropertyVisitorPath;
+struct FSceneStateExecutionContext;
 struct FSceneStateRange;
 struct FStructView;
 template<typename T> struct TObjectPtr;
@@ -31,13 +35,23 @@ namespace UE::SceneState
 		}
 	}
 
-	SCENESTATE_API TArray<FStructView> GetStructViews(FInstancedStructContainer& InStructContainer, FSceneStateRange InRange);
+	UE_API TArray<FStructView> GetStructViews(FInstancedStructContainer& InStructContainer, FSceneStateRange InRange);
+
+	UE_API TArray<FStructView> GetStructViews(FInstancedStructContainer& InStructContainer);
 
-	SCENESTATE_API TArray<FStructView> GetStructViews(FInstancedStructContainer& InStructContainer);
+	UE_API TArray<FConstStructView> GetConstStructViews(const FInstancedStructContainer& InStructContainer, FSceneStateRange InRange);
 
-	SCENESTATE_API TArray<FConstStructView> GetConstStructViews(const FInstancedStructContainer& InStructContainer, FSceneStateRange InRange);
+	UE_API TArray<FConstStructView> GetConstStructViews(const FInstancedStructContainer& InStructContainer);
 
-	SCENESTATE_API TArray<FConstStructView> GetConstStructViews(const FInstancedStructContainer& InStructContainer);
+	struct FApplyBatchParams
+	{
+		/** Batch of the bindings to apply */
+		uint16 BindingsBatch;
+		/** Target Data to copy the source data to */
+		FPropertyBindingDataView TargetDataView;
+	};
+	/** Applies a given binding batch to a target view within an execution context */
+	bool ApplyBatch(const FSceneStateExecutionContext& InContext, const FApplyBatchParams& InParams);
 
 	/**
 	 * Given a FPropertyVisitorPath, attempt to resolve that to a valid data pointer.
@@ -45,13 +59,13 @@ namespace UE::SceneState
 	 * TODO: FPropertyVisitorPath::GetPropertyDataPtr only takes a UObject ptr.
 	 * @see FPropertyVisitorPath::GetPropertyDataPtr
 	 */
-	SCENESTATE_API void* ResolveVisitedPath(const UScriptStruct* InRootObject, void* InRootData, const FPropertyVisitorPath& InPath);
+	UE_API void* ResolveVisitedPath(const UScriptStruct* InRootObject, void* InRootData, const FPropertyVisitorPath& InPath);
 
 	/**
 	 * Discards a given object by changing its outer to transient package and adding a prefix to it while also ensuring uniqueness
 	 * @param InObjectToDiscard the object to discard
 	 */
-	SCENESTATE_API void DiscardObject(UObject* InObjectToDiscard);
+	UE_API void DiscardObject(UObject* InObjectToDiscard);
 
 	/**
 	 * Discards an object with a given name in the given outer,
@@ -61,7 +75,7 @@ namespace UE::SceneState
 	 * @param InOnPreDiscardOldObject callback for when the old object is about to be trashed
 	 * @return the object that was discarded, or null if none was found
 	 */
-	SCENESTATE_API UObject* DiscardObject(UObject* InOuter
+	UE_API UObject* DiscardObject(UObject* InOuter
 		, const TCHAR* InObjectName
 		, TFunctionRef<void(UObject*)> InOnPreDiscardOldObject = [](UObject*){});
 
@@ -75,7 +89,7 @@ namespace UE::SceneState
 	 * @param InOnPreDiscardOldObject callback for when the old object is about to be trashed
 	 * @return true if the operation took place, false otherwise
 	 */
-	SCENESTATE_API bool ReplaceObject(UObject*& InOutObject
+	UE_API bool ReplaceObject(UObject*& InOutObject
 		, UObject* InOuter
 		, UClass* InClass
 		, const TCHAR* InObjectName
@@ -83,3 +97,5 @@ namespace UE::SceneState
 		, TFunctionRef<void(UObject*)> InOnPreDiscardOldObject = [](UObject*){});
 
 } // UE::SceneState
+
+#undef UE_API
diff --git a/Engine/Plugins/VirtualProduction/SceneState/Source/SceneStateBinding/Private/SceneStateBindingUtils.cpp b/Engine/Plugins/VirtualProduction/SceneState/Source/SceneStateBinding/Private/SceneStateBindingUtils.cpp
--- a/Engine/Plugins/VirtualProduction/SceneState/Source/SceneStateBinding/Private/SceneStateBindingUtils.cpp
+++ b/Engine/Plugins/VirtualProduction/SceneState/Source/SceneStateBinding/Private/SceneStateBindingUtils.cpp
@@ -63,28 +63,6 @@ void HandleStructIdChanged(UObject& InObject, const FGuid& InOldStructId, const
 }
 #endif
 
-bool ApplyBatch(const FApplyBatchParams& InParams)
-{
-	const FPropertyBindingIndex16 BatchIndex(InParams.BatchIndex);
-	if (!BatchIndex.IsValid())
-	{
-		return false;
-	}
-
-	bool bResult = true;
-
-	const FPropertyBindingCopyInfoBatch& Batch = InParams.BindingCollection.GetBatch(BatchIndex);
-	check(InParams.TargetDataView.GetStruct() == Batch.TargetStruct.Get().Struct);
-
-	for (const FPropertyBindingCopyInfo& Copy : InParams.BindingCollection.GetBatchCopies(Batch))
-	{
-		const FPropertyBindingDataView SourceView = InParams.FindSourceDataViewFunctor(Copy.SourceDataHandle.Get<FSceneStateBindingDataHandle>());
-		bResult &= InParams.BindingCollection.CopyProperty(Copy, SourceView, InParams.TargetDataView);
-	}
-
-	return bResult;
-}
-
 void PatchBindingCollection(const FPatchBindingParams& InParams)
 {
 	PatchBindingDescs(InParams);
diff --git a/Engine/Plugins/VirtualProduction/SceneState/Source/SceneStateBinding/Public/SceneStateBindingUtils.h b/Engine/Plugins/VirtualProduction/SceneState/Source/SceneStateBinding/Public/SceneStateBindingUtils.h
--- a/Engine/Plugins/VirtualProduction/SceneState/Source/SceneStateBinding/Public/SceneStateBindingUtils.h
+++ b/Engine/Plugins/VirtualProduction/SceneState/Source/SceneStateBinding/Public/SceneStateBindingUtils.h
@@ -24,20 +24,6 @@ namespace UE::SceneState
 	SCENESTATEBINDING_API void HandleStructIdChanged(UObject& InObject, const FGuid& InOldStructId, const FGuid& InNewStructId);
 #endif
 
-	struct FApplyBatchParams
-	{
-		/** Batch within the Binding Collection to run */
-		uint16 BatchIndex;
-		/** Target Data to copy the source data to */
-		FPropertyBindingDataView TargetDataView;
-		/** Binding collection holding the Copy Batch to run */
-		const FSceneStateBindingCollection& BindingCollection;
-		/** Functor to find the Source Data View for a given Data Handle */
-		TFunctionRef<FPropertyBindingDataView(const FSceneStateBindingDataHandle&)> FindSourceDataViewFunctor;
-	};
-	/** Runs the Batch within the Binding Collection for a given Target Data */
-	SCENESTATEBINDING_API bool ApplyBatch(const FApplyBatchParams& InParams);
-
 	struct FPatchBindingParams
 	{
 		/** Binding collection to patch */
