Motion Design Scene State: Added property functions (just known as functions within scene state for simplicity).
Functions are layer of logic that execute prior to evaluating owner's bindings.
For example, 'Set Integer' task could now be combined with an 'Integer Add' function so that the setter is no longer exclusively an A=B operation, but turned into a A=B+C operation.

diff --git a/Engine/Plugins/VirtualProduction/SceneState/Source/SceneState/Private/Functions/SceneStateFunction.cpp b/Engine/Plugins/VirtualProduction/SceneState/Source/SceneState/Private/Functions/SceneStateFunction.cpp
new file mode 100644
--- /dev/null
+++ b/Engine/Plugins/VirtualProduction/SceneState/Source/SceneState/Private/Functions/SceneStateFunction.cpp
@@ -0,0 +1,41 @@
+// Copyright Epic Games, Inc. All Rights Reserved.
+
+#include "Functions/SceneStateFunction.h"
+#include "SceneStateExecutionContext.h"
+#include "SceneStateUtils.h"
+
+#if WITH_EDITOR
+const UScriptStruct* FSceneStateFunction::GetFunctionDataType() const
+{
+	return OnGetFunctionDataType();
+}
+#endif
+
+void FSceneStateFunction::Setup(const FSceneStateExecutionContext& InContext) const
+{
+	InContext.SetupFunctionInstances(BindingsBatch);
+}
+
+void FSceneStateFunction::Execute(const FSceneStateExecutionContext& InContext, FStructView InFunctionInstance) const
+{
+	ApplyBindings(InContext, InFunctionInstance);
+	OnExecute(InContext, InFunctionInstance);
+}
+
+void FSceneStateFunction::Cleanup(const FSceneStateExecutionContext& InContext) const
+{
+	InContext.RemoveFunctionInstances(BindingsBatch);
+}
+
+bool FSceneStateFunction::ApplyBindings(const FSceneStateExecutionContext& InContext, FStructView InFunctionInstance) const
+{
+	QUICK_SCOPE_CYCLE_COUNTER(SceneStateFunction_ApplyBindings);
+
+	const UE::SceneState::FApplyBatchParams ApplyBatchParams
+		{
+			.BindingsBatch = BindingsBatch.Get(),
+			.TargetDataView = InFunctionInstance,
+		};
+
+	return ApplyBatch(InContext, ApplyBatchParams);
+}
diff --git a/Engine/Plugins/VirtualProduction/SceneState/Source/SceneState/Private/SceneStateExecutionContext.cpp b/Engine/Plugins/VirtualProduction/SceneState/Source/SceneState/Private/SceneStateExecutionContext.cpp
--- a/Engine/Plugins/VirtualProduction/SceneState/Source/SceneState/Private/SceneStateExecutionContext.cpp
+++ b/Engine/Plugins/VirtualProduction/SceneState/Source/SceneState/Private/SceneStateExecutionContext.cpp
@@ -1,6 +1,7 @@
 // Copyright Epic Games, Inc. All Rights Reserved.
 
 #include "SceneStateExecutionContext.h"
+#include "Functions/SceneStateFunction.h"
 #include "PropertyBindingDataView.h"
 #include "PropertyBindingTypes.h"
 #include "SceneState.h"
@@ -161,6 +162,9 @@ FPropertyBindingDataView FSceneStateExecutionContext::FindDataView(const FSceneS
 			return StateMachineInstance->Parameters.GetMutableValue();
 		}
 		break;
+
+	case ESceneStateDataType::Function:
+		return FindFunctionInstance(DataIndex);
 	}
 
 	// The data handle could be exclusively on the parent context, so call FindDataView on the parent too
@@ -707,6 +711,122 @@ void FSceneStateExecutionContext::RemoveConduitInstance(const FSceneStateConduit
 	}
 }
 
+bool FSceneStateExecutionContext::SetupFunctionInstances(FPropertyBindingIndex16 InBindingsBatch) const
+{
+	if (!InBindingsBatch.IsValid())
+	{
+		// Normal to not have a valid binding batch.
+		return false;
+	}
+
+	const FSceneStateBindingCollection* const BindingCollection = GetBindingCollection();
+	if (!BindingCollection)
+	{
+		UE_LOG(LogSceneState, Error, TEXT("[%s] Failed to setup function instances for batch %d. Binding Collection is null")
+			, *GetExecutionContextName()
+			, InBindingsBatch.Get());
+		return false;
+	}
+
+	const FPropertyBindingCopyInfoBatch& CopyInfoBatch = BindingCollection->GetBatch(InBindingsBatch);
+
+	if (!UE::SceneState::IsValidRange(CopyInfoBatch.PropertyFunctionsBegin, CopyInfoBatch.PropertyFunctionsEnd))
+	{
+		// Binding batch could be valid, but might not have functions bound.
+		return false;
+	}
+
+	FSceneStateFunctionInstanceBatch& FunctionInstanceBatch = FunctionInstanceBatches.Add(InBindingsBatch);
+	FunctionInstanceBatch.FunctionRange = FSceneStateRange::MakeBeginEndRange(CopyInfoBatch.PropertyFunctionsBegin.Get(), CopyInfoBatch.PropertyFunctionsEnd.Get());
+
+	// Initialize the function instances to the templates in the generated class
+	{
+		TArray<FConstStructView> TemplateFunctionInstances;
+		TemplateFunctionInstances.Reserve(FunctionInstanceBatch.FunctionRange.Count);
+
+		for (uint16 FunctionIndex = FunctionInstanceBatch.FunctionRange.Index; FunctionIndex <= FunctionInstanceBatch.FunctionRange.GetLastIndex(); ++FunctionIndex)
+		{
+			TemplateFunctionInstances.Add(GeneratedClass->GetFunctionInstances()[FunctionIndex]);
+		}
+
+		FunctionInstanceBatch.FunctionInstances.Append(TemplateFunctionInstances);
+	}
+
+	// Store a map of the absolute indices to the struct views in the batch for quick access from a single index (rather than going through batch index -> relative index)
+	for (uint16 FunctionIndex = FunctionInstanceBatch.FunctionRange.Index; FunctionIndex <= FunctionInstanceBatch.FunctionRange.GetLastIndex(); ++FunctionIndex)
+	{
+		const uint16 RelativeIndex = FunctionIndex - FunctionInstanceBatch.FunctionRange.Index;
+		FunctionInstanceViews.Add(FunctionIndex, FunctionInstanceBatch.FunctionInstances[RelativeIndex]);
+	}
+
+	// Set up these functions as they could be bound to other functions
+	for (uint16 FunctionIndex = FunctionInstanceBatch.FunctionRange.Index; FunctionIndex <= FunctionInstanceBatch.FunctionRange.GetLastIndex(); ++FunctionIndex)
+	{
+		if (const FSceneStateFunction* Function = GeneratedClass->GetFunctions()[FunctionIndex].GetPtr<const FSceneStateFunction>())
+		{
+			Function->Setup(*this);
+		}
+	}
+
+	return true;
+}
+
+void FSceneStateExecutionContext::RemoveFunctionInstances(FPropertyBindingIndex16 InBindingsBatch) const
+{
+	if (!InBindingsBatch.IsValid())
+	{
+		return;
+	}
+
+	FSceneStateRange FunctionRange;
+	if (FSceneStateFunctionInstanceBatch* Batch = FunctionInstanceBatches.Find(InBindingsBatch))
+	{
+		FunctionRange = Batch->FunctionRange;
+	}
+	else
+	{
+		return;
+	}
+
+	// Clean up the functions used resources (e.g. other function instance data)
+	for (uint16 FunctionIndex = FunctionRange.Index; FunctionIndex <= FunctionRange.GetLastIndex(); ++FunctionIndex)
+	{
+		if (const FSceneStateFunction* Function = GeneratedClass->GetFunctions()[FunctionIndex].GetPtr<const FSceneStateFunction>())
+		{
+			Function->Cleanup(*this);
+		}
+	}
+
+	// Remove the allocated instance data for each function
+	for (uint16 FunctionIndex = FunctionRange.Index; FunctionIndex <= FunctionRange.GetLastIndex(); ++FunctionIndex)
+	{
+		FunctionInstanceViews.Remove(FunctionIndex);
+	}
+	FunctionInstanceBatches.Remove(InBindingsBatch);
+}
+
+FConstStructView FSceneStateExecutionContext::FindFunction(uint16 InAbsoluteIndex) const
+{
+	if (!GeneratedClass)
+	{
+		return {};
+	}
+
+	const FInstancedStructContainer& Functions = GeneratedClass->GetFunctions();
+
+	if (!Functions.IsValidIndex(InAbsoluteIndex))
+	{
+		return {};
+	}
+
+	return Functions[InAbsoluteIndex];
+}
+
+FStructView FSceneStateExecutionContext::FindFunctionInstance(uint16 InAbsoluteIndex) const
+{
+	return FunctionInstanceViews.FindRef(InAbsoluteIndex);
+}
+
 const FSceneStateExecutionContext* FSceneStateExecutionContext::FindParentContext() const
 {
 	if (TSharedPtr<UE::SceneState::FExecutionContextRegistry> ContextRegistry = ContextRegistryWeak.Pin())
diff --git a/Engine/Plugins/VirtualProduction/SceneState/Source/SceneState/Private/SceneStateGeneratedClass.cpp b/Engine/Plugins/VirtualProduction/SceneState/Source/SceneState/Private/SceneStateGeneratedClass.cpp
--- a/Engine/Plugins/VirtualProduction/SceneState/Source/SceneState/Private/SceneStateGeneratedClass.cpp
+++ b/Engine/Plugins/VirtualProduction/SceneState/Source/SceneState/Private/SceneStateGeneratedClass.cpp
@@ -202,6 +202,9 @@ const UStruct* USceneStateGeneratedClass::FindDataStruct(const FSceneStateBindin
 
 	case ESceneStateDataType::StateMachine:
 		return StateMachines[DataIndex].GetParametersStruct();
+
+	case ESceneStateDataType::Function:
+		return FunctionInstances[DataIndex].GetScriptStruct();
 	}
 
 	return nullptr;
@@ -232,6 +235,8 @@ void USceneStateGeneratedClass::Reset()
 	Tasks.Reset();
 	TaskPrerequisites.Reset();
 	TaskInstances.Reset();
+	Functions.Reset();
+	FunctionInstances.Reset();
 	States.Reset();
 	Conduits.Reset();
 	ConduitLinks.Reset();
@@ -245,6 +250,7 @@ void USceneStateGeneratedClass::Reset()
 
 #if WITH_EDITOR
 	TaskMetadata.Reset();
+	FunctionMetadata.Reset();
 	StateMetadata.Reset();
 	TransitionMetadata.Reset();
 	StateNodeToIndex.Reset();
diff --git a/Engine/Plugins/VirtualProduction/SceneState/Source/SceneState/Private/SceneStateMachine.cpp b/Engine/Plugins/VirtualProduction/SceneState/Source/SceneState/Private/SceneStateMachine.cpp
--- a/Engine/Plugins/VirtualProduction/SceneState/Source/SceneState/Private/SceneStateMachine.cpp
+++ b/Engine/Plugins/VirtualProduction/SceneState/Source/SceneState/Private/SceneStateMachine.cpp
@@ -28,6 +28,8 @@ void FSceneStateMachine::Setup(const FSceneStateExecutionContext& InContext) con
 		return;
 	}
 
+	InContext.SetupFunctionInstances(BindingsBatch);
+
 	// Initialize Instance
 	Instance->ActiveIndex = EntryIndex;
 	Instance->Parameters = Parameters;
@@ -119,6 +121,7 @@ void FSceneStateMachine::Stop(const FSceneStateExecutionContext& InContext) cons
 
 	StopActiveState(InContext, *Instance);
 
+	InContext.RemoveFunctionInstances(BindingsBatch);
 	InContext.RemoveStateMachineInstance(*this);
 }
 
diff --git a/Engine/Plugins/VirtualProduction/SceneState/Source/SceneState/Private/SceneStateRange.cpp b/Engine/Plugins/VirtualProduction/SceneState/Source/SceneState/Private/SceneStateRange.cpp
new file mode 100644
--- /dev/null
+++ b/Engine/Plugins/VirtualProduction/SceneState/Source/SceneState/Private/SceneStateRange.cpp
@@ -0,0 +1,16 @@
+// Copyright Epic Games, Inc. All Rights Reserved.
+
+#include "SceneStateRange.h"
+
+FSceneStateRange FSceneStateRange::MakeBeginEndRange(uint16 InBegin, uint16 InEnd)
+{
+	if (InBegin == InvalidIndex || InEnd == InvalidIndex || InBegin >= InEnd)
+	{
+		return FSceneStateRange();
+	}
+
+	FSceneStateRange Range;
+	Range.Index = InBegin;
+	Range.Count = InEnd - InBegin;
+	return Range;
+}
diff --git a/Engine/Plugins/VirtualProduction/SceneState/Source/SceneState/Private/SceneStateUtils.cpp b/Engine/Plugins/VirtualProduction/SceneState/Source/SceneState/Private/SceneStateUtils.cpp
--- a/Engine/Plugins/VirtualProduction/SceneState/Source/SceneState/Private/SceneStateUtils.cpp
+++ b/Engine/Plugins/VirtualProduction/SceneState/Source/SceneState/Private/SceneStateUtils.cpp
@@ -2,6 +2,7 @@
 
 #include "SceneStateUtils.h"
 #include "Engine/Engine.h"
+#include "Functions/SceneStateFunction.h"
 #include "PropertyBindingPath.h"
 #include "PropertyBindingTypes.h"
 #include "SceneStateBindingCollection.h"
@@ -67,6 +68,11 @@ TArray<FConstStructView> GetConstStructViews(const FInstancedStructContainer& In
 	return GetConstStructViews(InStructContainer, Range);
 }
 
+bool IsValidRange(FPropertyBindingIndex16 InBegin, FPropertyBindingIndex16 InEnd)
+{
+	return InBegin.IsValid() && InEnd.IsValid() && InBegin.Get() < InEnd.Get();
+}
+
 bool ApplyBatch(const FSceneStateExecutionContext& InContext, const FApplyBatchParams& InParams)
 {
 	const FPropertyBindingIndex16 BindingsBatch(InParams.BindingsBatch);
@@ -86,6 +92,21 @@ bool ApplyBatch(const FSceneStateExecutionContext& InContext, const FApplyBatchP
 	const FPropertyBindingCopyInfoBatch& Batch = BindingCollection->GetBatch(BindingsBatch);
 	check(InParams.TargetDataView.GetStruct() == Batch.TargetStruct.Get().Struct);
 
+	// If there were valid functions found on setup, execute them now
+	if (UE::SceneState::IsValidRange(Batch.PropertyFunctionsBegin, Batch.PropertyFunctionsEnd))
+	{
+		const FSceneStateRange FunctionRange = FSceneStateRange::MakeBeginEndRange(Batch.PropertyFunctionsBegin.Get(), Batch.PropertyFunctionsEnd.Get());
+
+		for (uint16 FunctionIndex = FunctionRange.Index; FunctionIndex <= FunctionRange.GetLastIndex(); ++FunctionIndex)
+		{
+			if (const FSceneStateFunction* Function = InContext.FindFunction(FunctionIndex).GetPtr<const FSceneStateFunction>())
+			{
+				FStructView FunctionInstance = InContext.FindFunctionInstance(FunctionIndex);
+				Function->Execute(InContext, FunctionInstance);
+			}
+		}
+	}
+
 	for (const FPropertyBindingCopyInfo& Copy : BindingCollection->GetBatchCopies(Batch))
 	{
 		const FPropertyBindingDataView SourceView = InContext.FindDataView(Copy.SourceDataHandle.Get<FSceneStateBindingDataHandle>());
diff --git a/Engine/Plugins/VirtualProduction/SceneState/Source/SceneState/Private/Tasks/SceneStateTask.cpp b/Engine/Plugins/VirtualProduction/SceneState/Source/SceneState/Private/Tasks/SceneStateTask.cpp
--- a/Engine/Plugins/VirtualProduction/SceneState/Source/SceneState/Private/Tasks/SceneStateTask.cpp
+++ b/Engine/Plugins/VirtualProduction/SceneState/Source/SceneState/Private/Tasks/SceneStateTask.cpp
@@ -77,6 +77,7 @@ void FSceneStateTask::Setup(const FSceneStateExecutionContext& InContext, FStruc
 	}
 
 	Instance->SetStatus(EExecutionStatus::NotStarted);
+	InContext.SetupFunctionInstances(BindingsBatch);
 	OnSetup(InContext, InTaskInstance);
 }
 
@@ -137,6 +138,7 @@ void FSceneStateTask::Stop(const FSceneStateExecutionContext& InContext, FStruct
 		OnStop(InContext, InTaskInstance, InStopReason);
 	}
 
+	InContext.RemoveFunctionInstances(BindingsBatch);
 	Instance->SetStatus(Private::GetTaskStatus(InStopReason));
 }
 
diff --git a/Engine/Plugins/VirtualProduction/SceneState/Source/SceneState/Private/Transition/SceneStateTransition.cpp b/Engine/Plugins/VirtualProduction/SceneState/Source/SceneState/Private/Transition/SceneStateTransition.cpp
--- a/Engine/Plugins/VirtualProduction/SceneState/Source/SceneState/Private/Transition/SceneStateTransition.cpp
+++ b/Engine/Plugins/VirtualProduction/SceneState/Source/SceneState/Private/Transition/SceneStateTransition.cpp
@@ -41,6 +41,7 @@ void FSceneStateTransition::Setup(const FSceneStateExecutionContext& InContext)
 		return;
 	}
 
+	InContext.SetupFunctionInstances(BindingsBatch);
 	Instance->Parameters = InContext.GetTemplateTransitionParameter(*this);
 
 	// Setup the conduit chaining this transition
@@ -71,6 +72,7 @@ bool FSceneStateTransition::Evaluate(const UE::SceneState::FTransitionEvaluation
 void FSceneStateTransition::Exit(const FSceneStateExecutionContext& InContext) const
 {
 	InContext.RemoveTransitionInstance(*this);
+	InContext.RemoveFunctionInstances(BindingsBatch);
 
 	// Exit the conduit chaining this transition
 	if (const FSceneStateConduit* Conduit = FindConduitTarget(InContext))
diff --git a/Engine/Plugins/VirtualProduction/SceneState/Source/SceneState/Public/Functions/SceneStateFunction.h b/Engine/Plugins/VirtualProduction/SceneState/Source/SceneState/Public/Functions/SceneStateFunction.h
new file mode 100644
--- /dev/null
+++ b/Engine/Plugins/VirtualProduction/SceneState/Source/SceneState/Public/Functions/SceneStateFunction.h
@@ -0,0 +1,99 @@
+// Copyright Epic Games, Inc. All Rights Reserved.
+
+#pragma once
+
+#include "PropertyBindingTypes.h"
+#include "SceneStateFunction.generated.h"
+
+#define UE_API SCENESTATE_API
+
+struct FSceneStateExecutionContext;
+struct FStructView;
+
+namespace UE::SceneState
+{
+	namespace Editor
+	{
+		class FBindingCompiler;
+	}
+}
+
+/**
+ * Base struct for all scene state functions.
+ * Binding Functions are a layer of logic that execute prior to evaluating owner's bindings.
+ *
+ * The function's instance data is expected to have its output properties with meta "Output", and these will be hidden in the UI.
+ * This property is used to find which bindings the function can be used for.
+ *
+ * Example:
+ *
+ *	USTRUCT()
+ *	struct FTextToStringFunctionInstance
+ *	{
+ *		GENERATED_BODY()
+ *
+ *		UPROPERTY(EditAnywhere, Category="Test")
+ *		FText Text;
+ *
+ *		UPROPERTY(EditAnywhere, Category="Test", meta=(Output))
+ *		FString Result;
+ *	};
+ */
+USTRUCT(meta=(Hidden))
+struct FSceneStateFunction
+{
+	GENERATED_BODY()
+
+#if WITH_EDITOR
+	/**
+	 * Called in-editor to get the function data type
+	 * @return the function data struct type
+	 */
+	 UE_API const UScriptStruct* GetFunctionDataType() const;
+#endif
+
+	/** Called when binding owners are setting up and allocating instance data */
+	void Setup(const FSceneStateExecutionContext& InContext) const;
+
+	/**
+	 * Called right before evaluating bindings for the owner.
+	 * @param InContext reference to current execution context.
+	 * @param InFunctionInstance data view of the function instance data
+	 */
+	void Execute(const FSceneStateExecutionContext& InContext, FStructView InFunctionInstance) const;
+
+	/** Called when binding owners are cleaning up their allocated instance data */
+	void Cleanup(const FSceneStateExecutionContext& InContext) const;
+
+protected:
+#if WITH_EDITOR
+	/** Called to retrieve the data type that this function expects*/
+	virtual const UScriptStruct* OnGetFunctionDataType() const PURE_VIRTUAL(FSceneStateFunction::OnGetFunctionDataType, return nullptr;);
+#endif
+
+	/** Called right before evaluating bindings for the owner. */
+	virtual void OnExecute(const FSceneStateExecutionContext& InContext, FStructView InFunctionInstance) const
+	{
+	}
+
+private:
+	/** Applies the bindings batch to the function instance */
+	bool ApplyBindings(const FSceneStateExecutionContext& InContext, FStructView InFunctionInstance) const;
+
+	/** Bindings Batch where this function is target */
+	UPROPERTY()
+	FPropertyBindingIndex16 BindingsBatch;
+
+	friend UE::SceneState::Editor::FBindingCompiler;
+};
+
+template<>
+struct TStructOpsTypeTraits<FSceneStateFunction> : TStructOpsTypeTraitsBase2<FSceneStateFunction>
+{
+	enum
+	{
+		WithPureVirtual = true,
+	};
+};
+
+#undef UE_API
diff --git a/Engine/Plugins/VirtualProduction/SceneState/Source/SceneState/Public/Functions/SceneStateFunctionInstanceBatch.h b/Engine/Plugins/VirtualProduction/SceneState/Source/SceneState/Public/Functions/SceneStateFunctionInstanceBatch.h
new file mode 100644
--- /dev/null
+++ b/Engine/Plugins/VirtualProduction/SceneState/Source/SceneState/Public/Functions/SceneStateFunctionInstanceBatch.h
@@ -0,0 +1,22 @@
+// Copyright Epic Games, Inc. All Rights Reserved.
+
+#pragma once
+
+#include "SceneStateRange.h"
+#include "StructUtils/InstancedStructContainer.h"
+#include "SceneStateFunctionInstanceBatch.generated.h"
+
+/** Holds all the function instances for a given batch */
+USTRUCT()
+struct FSceneStateFunctionInstanceBatch
+{
+	GENERATED_BODY()
+
+	/** The function range (in absolute index) of the functions within the batch */
+	UPROPERTY()
+	FSceneStateRange FunctionRange;
+
+	/** Container holding all the function instances for a batch in a contiguous block of memory */
+	UPROPERTY()
+	FInstancedStructContainer FunctionInstances;
+};
diff --git a/Engine/Plugins/VirtualProduction/SceneState/Source/SceneState/Public/Functions/SceneStateFunctionMetadata.h b/Engine/Plugins/VirtualProduction/SceneState/Source/SceneState/Public/Functions/SceneStateFunctionMetadata.h
new file mode 100644
--- /dev/null
+++ b/Engine/Plugins/VirtualProduction/SceneState/Source/SceneState/Public/Functions/SceneStateFunctionMetadata.h
@@ -0,0 +1,18 @@
+// Copyright Epic Games, Inc. All Rights Reserved.
+
+#pragma once
+
+#include "Misc/Guid.h"
+#include "SceneStateFunctionMetadata.generated.h"
+
+/** Metadata information about the function. Its data is meant to be available/accessible in editor-only builds */
+USTRUCT()
+struct FSceneStateFunctionMetadata
+{
+	GENERATED_BODY()
+
+#if WITH_EDITORONLY_DATA
+	UPROPERTY()
+	FGuid FunctionId;
+#endif
+};
diff --git a/Engine/Plugins/VirtualProduction/SceneState/Source/SceneState/Public/SceneStateExecutionContext.h b/Engine/Plugins/VirtualProduction/SceneState/Source/SceneState/Public/SceneStateExecutionContext.h
--- a/Engine/Plugins/VirtualProduction/SceneState/Source/SceneState/Public/SceneStateExecutionContext.h
+++ b/Engine/Plugins/VirtualProduction/SceneState/Source/SceneState/Public/SceneStateExecutionContext.h
@@ -4,6 +4,7 @@
 
 #include "Conduit/SceneStateConduitInstance.h"
 #include "Containers/Array.h"
+#include "Functions/SceneStateFunctionInstanceBatch.h"
 #include "SceneStateExecutionContextHandle.h"
 #include "SceneStateInstance.h"
 #include "SceneStateMachineInstance.h"
@@ -25,6 +26,7 @@ struct FSceneStateBindingCollection;
 struct FSceneStateBindingDataHandle;
 struct FSceneStateConduit;
 struct FSceneStateEventHandler;
+struct FSceneStateFunction;
 struct FSceneStateMachine;
 struct FSceneStateMetadata;
 struct FSceneStatePropertyReference;
@@ -222,6 +224,21 @@ struct FSceneStateExecutionContext
 	/** Removes the conduit instance of this context for the given conduit */
 	void RemoveConduitInstance(const FSceneStateConduit& InConduit) const;
 
+	/**
+	 * Allocates the function instance data for a given binding batch
+	 * @return true if there were valid functions within the batch, false otherwise
+	 */
+	bool SetupFunctionInstances(FPropertyBindingIndex16 InBindingsBatch) const;
+
+	/** Removes the function instance data for a given binding batch */
+	void RemoveFunctionInstances(FPropertyBindingIndex16 InBindingsBatch) const;
+
+	/** Returns the function of this context for the given function absolute index. Returns an invalid view if not found */
+	FConstStructView FindFunction(uint16 InAbsoluteIndex) const;
+
+	/** Returns the function instance of this context for the given function absolute index. Returns an invalid view if not found */
+	FStructView FindFunctionInstance(uint16 InAbsoluteIndex) const;
+
 private:
 	/** Retrieves the parent execution context if existing. Can be null */
 	const FSceneStateExecutionContext* FindParentContext() const;
@@ -284,6 +301,19 @@ private:
 	UPROPERTY()
 	mutable TMap<uint16, FSceneStateConduitInstance> ConduitInstances;
 
+	/**
+	 * Map of the batch index to the function instances within the batch
+	 * It's allocated when its owner is setup, and removed on the owners exit
+	 */
+	UPROPERTY()
+	mutable TMap<FPropertyBindingIndex16, FSceneStateFunctionInstanceBatch> FunctionInstanceBatches;
+
+	/**
+	 * Map of the function absolute index to the view of the function instance
+	 * It's allocated when its owner is setup, and removed on the owners exit
+	 */
+	mutable TMap<uint16, FStructView> FunctionInstanceViews;
+
 	/** Context handle to the parent execution context. Can be invalid. */
 	UE::SceneState::FExecutionContextHandle ParentHandle;
 
diff --git a/Engine/Plugins/VirtualProduction/SceneState/Source/SceneState/Public/SceneStateGeneratedClass.h b/Engine/Plugins/VirtualProduction/SceneState/Source/SceneState/Public/SceneStateGeneratedClass.h
--- a/Engine/Plugins/VirtualProduction/SceneState/Source/SceneState/Public/SceneStateGeneratedClass.h
+++ b/Engine/Plugins/VirtualProduction/SceneState/Source/SceneState/Public/SceneStateGeneratedClass.h
@@ -5,6 +5,7 @@
 #include "Conduit/SceneStateConduit.h"
 #include "Conduit/SceneStateConduitLink.h"
 #include "Engine/BlueprintGeneratedClass.h"
+#include "Functions/SceneStateFunctionMetadata.h"
 #include "SceneState.h"
 #include "SceneStateBindingCollection.h"
 #include "SceneStateEventHandler.h"
@@ -26,6 +27,7 @@ class UUserDefinedStruct;
 namespace UE::SceneState::Editor
 {
 	class FBindingCompiler;
+	class FBindingFunctionCompiler;
 	class FBlueprintCompilerContext;
 	class FStateMachineCompiler;
 }
@@ -98,6 +100,16 @@ public:
 		return StateMachines;
 	}
 
+	const FInstancedStructContainer& GetFunctions() const
+	{
+		return Functions;
+	}
+
+	const FInstancedStructContainer& GetFunctionInstances() const
+	{
+		return FunctionInstances;
+	}
+
 #if WITH_EDITOR
 	TConstArrayView<FSceneStateMetadata> GetStateMetadata() const
 	{
@@ -187,6 +199,14 @@ private:
 	UPROPERTY()
 	FInstancedStructContainer TaskInstances;
 
+	/** Holds all the compiled functions (including from nested sources) in a contiguous block of memory */
+	UPROPERTY()
+	FInstancedStructContainer Functions;
+
+	/** Templates used to instantiate the function data */
+	UPROPERTY()
+	FInstancedStructContainer FunctionInstances;
+
 	/** Holds all the compiled event handlers (including from nested sources) contiguously */
 	UPROPERTY()
 	TArray<FSceneStateEventHandler> EventHandlers;
@@ -227,6 +247,10 @@ private:
 	UPROPERTY()
 	TArray<FSceneStateTaskMetadata> TaskMetadata;
 
+	/** Holds all the function editor-only metadata. This array matches in index with the functions. */
+	UPROPERTY()
+	TArray<FSceneStateFunctionMetadata> FunctionMetadata;
+
 	/** Holds all the transition editor-only metadata. This array matches in index with the transitions array. */
 	UPROPERTY()
 	TArray<FSceneStateTransitionMetadata> TransitionMetadata;
@@ -245,6 +269,7 @@ private:
 #endif
 
 	friend UE::SceneState::Editor::FBindingCompiler;
+	friend UE::SceneState::Editor::FBindingFunctionCompiler;
 	friend UE::SceneState::Editor::FBlueprintCompilerContext;
 	friend UE::SceneState::Editor::FStateMachineCompiler;
 	friend class USceneStateBlueprint;
diff --git a/Engine/Plugins/VirtualProduction/SceneState/Source/SceneState/Public/SceneStateRange.h b/Engine/Plugins/VirtualProduction/SceneState/Source/SceneState/Public/SceneStateRange.h
--- a/Engine/Plugins/VirtualProduction/SceneState/Source/SceneState/Public/SceneStateRange.h
+++ b/Engine/Plugins/VirtualProduction/SceneState/Source/SceneState/Public/SceneStateRange.h
@@ -6,6 +6,8 @@
 #include <limits>
 #include "SceneStateRange.generated.h"
 
+#define UE_API SCENESTATE_API
+
 USTRUCT()
 struct FSceneStateRange
 {
@@ -13,6 +15,9 @@ struct FSceneStateRange
 
 	static constexpr uint16 InvalidIndex = std::numeric_limits<uint16>::max();
 
+	/** Makes a range with the given begin and end range, where end is exclusive */
+	UE_API static FSceneStateRange MakeBeginEndRange(uint16 InBegin, uint16 InEnd);
+
 	uint16 GetLastIndex() const
 	{
 		checkSlow(IsValid());
@@ -30,3 +35,5 @@ struct FSceneStateRange
 	UPROPERTY()
 	uint16 Count = 0;
 };
+
+#undef UE_API
diff --git a/Engine/Plugins/VirtualProduction/SceneState/Source/SceneState/Public/SceneStateUtils.h b/Engine/Plugins/VirtualProduction/SceneState/Source/SceneState/Public/SceneStateUtils.h
--- a/Engine/Plugins/VirtualProduction/SceneState/Source/SceneState/Public/SceneStateUtils.h
+++ b/Engine/Plugins/VirtualProduction/SceneState/Source/SceneState/Public/SceneStateUtils.h
@@ -13,6 +13,7 @@ class UClass;
 class UObject;
 struct FConstStructView;
 struct FInstancedStructContainer;
+struct FPropertyBindingIndex16;
 struct FPropertyVisitorPath;
 struct FSceneStateExecutionContext;
 struct FSceneStateRange;
@@ -43,6 +44,17 @@ namespace UE::SceneState
 
 	UE_API TArray<FConstStructView> GetConstStructViews(const FInstancedStructContainer& InStructContainer);
 
+	/** Determines whether the given begin and end forms a valid range */
+	bool IsValidRange(FPropertyBindingIndex16 InBegin, FPropertyBindingIndex16 InEnd);
+
+	struct FSetupBatchParams
+	{
+		/** Batch of the bindings to apply */
+		uint16 BindingsBatch;
+	};
+	/** Sets up a given binding batch to then apply */
+	bool SetupBatch(const FSceneStateExecutionContext& InContext, const FSetupBatchParams& InParams);
+
 	struct FApplyBatchParams
 	{
 		/** Batch of the bindings to apply */
diff --git a/Engine/Plugins/VirtualProduction/SceneState/Source/SceneStateBinding/Private/SceneStateBindingCollection.cpp b/Engine/Plugins/VirtualProduction/SceneState/Source/SceneStateBinding/Private/SceneStateBindingCollection.cpp
--- a/Engine/Plugins/VirtualProduction/SceneState/Source/SceneStateBinding/Private/SceneStateBindingCollection.cpp
+++ b/Engine/Plugins/VirtualProduction/SceneState/Source/SceneStateBinding/Private/SceneStateBindingCollection.cpp
@@ -2,9 +2,25 @@
 
 #include "SceneStateBindingCollection.h"
 #include "PropertyBindingDataView.h"
+#include "SceneStateBindingFunction.h"
 #include "SceneStatePropertyReference.h"
 #include "SceneStatePropertyReferenceUtils.h"
 
+#if WITH_EDITOR
+FPropertyBindingPath FSceneStateBindingCollection::AddBindingFunction(const UE::SceneState::FBindingFunctionInfo& InFunctionInfo, TConstArrayView<FPropertyBindingPathSegment> InSourcePathSegments, const FPropertyBindingPath& InTargetPath)
+{
+	FInstancedStruct BindingFunction(FSceneStateBindingFunction::StaticStruct());
+	BindingFunction.InitializeAs<FSceneStateBindingFunction>(InFunctionInfo);
+
+	Super::RemoveBindings(InTargetPath, ESearchMode::Exact);
+
+	FPropertyBindingPath SourcePath = FPropertyBindingPath(BindingFunction.Get<FSceneStateBindingFunction>().FunctionId, InSourcePathSegments);
+	Bindings.Emplace(MoveTemp(BindingFunction), SourcePath, InTargetPath);
+
+	return SourcePath;
+}
+#endif
+
 const FSceneStateBindingDesc* FSceneStateBindingCollection::FindBindingDesc(FSceneStateBindingDataHandle InDataHandle) const
 {
 	return BindingDescs.FindByPredicate(
@@ -44,6 +60,56 @@ FPropertyBindingBinding* FSceneStateBindingCollection::AddBindingInternal(const
 	return &Bindings.Emplace_GetRef(InSourcePath, InTargetPath);
 }
 
+void FSceneStateBindingCollection::CopyBindingsInternal(const FGuid InFromStructId, const FGuid InToStructId)
+{
+	struct FFunctionBindingCopyInfo
+	{
+		UE::SceneState::FBindingFunctionInfo FunctionInfo;
+		FPropertyBindingPath SourcePath;
+		FPropertyBindingPath TargetPath;
+	};
+
+	TArray<TTuple<FGuid, FGuid>, TInlineAllocator<8>> StructIds;
+	StructIds.Emplace(InFromStructId, InToStructId);
+
+	// Worst case: all bindings have unique struct id and are added to the array
+	const int32 IndexCeiling = Bindings.Num();
+
+	// Find the StructId and copy the binding. If we find functions, then copy the function and copy the functions binding recursively.
+	for (int32 Index = 0; Index < StructIds.Num() && ensure(Index < IndexCeiling); ++Index)
+	{
+		const TTuple<FGuid, FGuid>& StructId = StructIds[Index];
+
+		TArray<FFunctionBindingCopyInfo, TInlineAllocator<8>> FunctionBindingsToCopy;
+
+		CopyBindingsImplementation(StructId.Get<0>(), StructId.Get<1>(),
+			[ToStructId = StructId.Get<1>(), &FunctionBindingsToCopy](const FPropertyBindingBinding& InBinding)
+			{
+				bool bCanCopyBinding = true;
+				if (const FSceneStateBindingFunction* BindingFunction = InBinding.GetPropertyFunctionNode().GetPtr<const FSceneStateBindingFunction>())
+				{
+					FFunctionBindingCopyInfo& CopyInfo = FunctionBindingsToCopy.AddDefaulted_GetRef();
+					CopyInfo.FunctionInfo.FunctionType = BindingFunction->Function.GetScriptStruct();
+					CopyInfo.FunctionInfo.InstanceType = BindingFunction->FunctionInstance.GetScriptStruct();
+					CopyInfo.SourcePath = InBinding.GetSourcePath();
+					CopyInfo.TargetPath = FPropertyBindingPath(ToStructId, InBinding.GetTargetPath().GetSegments());
+
+					bCanCopyBinding = false;
+				}
+				return bCanCopyBinding;
+			});
+
+		// Copy all functions bindings that target "FromStructId" and retarget them to "ToStructId".
+		for (const FFunctionBindingCopyInfo& CopyInfo : FunctionBindingsToCopy)
+		{
+			const FGuid NewId = AddBindingFunction(CopyInfo.FunctionInfo, CopyInfo.SourcePath.GetSegments(), CopyInfo.TargetPath).GetStructID();
+
+			// Adds the new function binding to the list of copy bindings to look for.
+			StructIds.AddUnique(TTuple<FGuid, FGuid>{CopyInfo.SourcePath.GetStructID(), NewId});
+		}
+	}
+}
+
 void FSceneStateBindingCollection::RemoveBindingsInternal(TFunctionRef<bool(FPropertyBindingBinding&)> InPredicate)
 {
 	Bindings.RemoveAllSwap(InPredicate);
diff --git a/Engine/Plugins/VirtualProduction/SceneState/Source/SceneStateBinding/Private/SceneStateBindingFunction.cpp b/Engine/Plugins/VirtualProduction/SceneState/Source/SceneStateBinding/Private/SceneStateBindingFunction.cpp
new file mode 100644
--- /dev/null
+++ b/Engine/Plugins/VirtualProduction/SceneState/Source/SceneStateBinding/Private/SceneStateBindingFunction.cpp
@@ -0,0 +1,17 @@
+// Copyright Epic Games, Inc. All Rights Reserved.
+
+#include "SceneStateBindingFunction.h"
+
+#if WITH_EDITOR
+FSceneStateBindingFunction::FSceneStateBindingFunction(const UE::SceneState::FBindingFunctionInfo& InFunctionInfo)
+{
+	FunctionId = FGuid::NewGuid();
+	Function.InitializeAs(InFunctionInfo.FunctionType);
+	FunctionInstance.InitializeAs(InFunctionInfo.InstanceType);
+}
+
+bool FSceneStateBindingFunction::IsValid() const
+{
+	return FunctionId.IsValid() && Function.IsValid() && FunctionInstance.IsValid();
+}
+#endif
diff --git a/Engine/Plugins/VirtualProduction/SceneState/Source/SceneStateBinding/Public/SceneStateBindingCollection.h b/Engine/Plugins/VirtualProduction/SceneState/Source/SceneStateBinding/Public/SceneStateBindingCollection.h
--- a/Engine/Plugins/VirtualProduction/SceneState/Source/SceneStateBinding/Public/SceneStateBindingCollection.h
+++ b/Engine/Plugins/VirtualProduction/SceneState/Source/SceneStateBinding/Public/SceneStateBindingCollection.h
@@ -8,11 +8,18 @@
 #include "SceneStateBindingReference.h"
 #include "SceneStateBindingCollection.generated.h"
 
+#define UE_API SCENESTATEBINDING_API
+
 struct FSceneStatePropertyReference;
 
-namespace UE::SceneState::Editor
+namespace UE::SceneState
 {
-	class FBindingCompiler;
+	struct FBindingFunctionInfo;
+
+	namespace Editor
+	{
+		class FBindingCompiler;
+	}
 }
 
 /** Holds all the property bindings to use in execution */
@@ -31,37 +38,47 @@ struct FSceneStateBindingCollection : public FPropertyBindingBindingCollection
 		return Bindings;
 	}
 
+	TConstArrayView<FSceneStateBinding> GetBindings() const
+	{
+		return Bindings;
+	}
+
+#if WITH_EDITOR
+	UE_API FPropertyBindingPath AddBindingFunction(const UE::SceneState::FBindingFunctionInfo& InFunctionInfo, TConstArrayView<FPropertyBindingPathSegment> InSourcePathSegments, const FPropertyBindingPath& InTargetPath);
+#endif
+
 	/** Finds the binding desc matching the given data handle */
-	SCENESTATEBINDING_API const FSceneStateBindingDesc* FindBindingDesc(FSceneStateBindingDataHandle InDataHandle) const;
+	UE_API const FSceneStateBindingDesc* FindBindingDesc(FSceneStateBindingDataHandle InDataHandle) const;
 
 	/** Retrieves the resolved reference for the given property reference */
-	SCENESTATEBINDING_API const FSceneStateBindingResolvedReference* FindResolvedReference(const FSceneStatePropertyReference& InPropertyReference) const;
+	UE_API const FSceneStateBindingResolvedReference* FindResolvedReference(const FSceneStatePropertyReference& InPropertyReference) const;
 
 	/**
 	 * Resolves the property access for the given reference and data view
 	 * @note This does not deal with recursive resolution (for cases where the reference is to another reference).
 	 * Prefer using UE::SceneState::ResolveProperty instead.
 	 */
-	SCENESTATEBINDING_API uint8* ResolveProperty(const FSceneStateBindingResolvedReference& InResolvedReference, FPropertyBindingDataView InDataView) const;
+	UE_API uint8* ResolveProperty(const FSceneStateBindingResolvedReference& InResolvedReference, FPropertyBindingDataView InDataView) const;
 
 	//~ Begin FPropertyBindingBindingCollection
 #if WITH_EDITOR
-	SCENESTATEBINDING_API virtual FPropertyBindingBinding* AddBindingInternal(const FPropertyBindingPath& InSourcePath, const FPropertyBindingPath& InTargetPath) override;
-	SCENESTATEBINDING_API virtual void RemoveBindingsInternal(TFunctionRef<bool(FPropertyBindingBinding&)> InPredicate) override;
-	SCENESTATEBINDING_API virtual bool HasBindingInternal(TFunctionRef<bool(const FPropertyBindingBinding&)> InPredicate) const override;
-	SCENESTATEBINDING_API virtual const FPropertyBindingBinding* FindBindingInternal(TFunctionRef<bool(const FPropertyBindingBinding&)> InPredicate) const override;
+	UE_API virtual FPropertyBindingBinding* AddBindingInternal(const FPropertyBindingPath& InSourcePath, const FPropertyBindingPath& InTargetPath) override;
+	UE_API virtual void CopyBindingsInternal(const FGuid InFromStructId, const FGuid InToStructId) override;
+	UE_API virtual void RemoveBindingsInternal(TFunctionRef<bool(FPropertyBindingBinding&)> InPredicate) override;
+	UE_API virtual bool HasBindingInternal(TFunctionRef<bool(const FPropertyBindingBinding&)> InPredicate) const override;
+	UE_API virtual const FPropertyBindingBinding* FindBindingInternal(TFunctionRef<bool(const FPropertyBindingBinding&)> InPredicate) const override;
 #endif
-	SCENESTATEBINDING_API virtual int32 GetNumBindings() const override;
-	SCENESTATEBINDING_API virtual int32 GetNumBindableStructDescriptors() const override;
-	SCENESTATEBINDING_API virtual const FPropertyBindingBindableStructDescriptor* GetBindableStructDescriptorFromHandle(FConstStructView InSourceHandleView) const override;
-	SCENESTATEBINDING_API virtual void ForEachBinding(TFunctionRef<void(const FPropertyBindingBinding& Binding)> InFunction) const override;
-	SCENESTATEBINDING_API virtual void ForEachBinding(FPropertyBindingIndex16 InBegin, FPropertyBindingIndex16 InEnd, TFunctionRef<void(const FPropertyBindingBinding& Binding, const int32 BindingIndex)> InFunction) const override;
-	SCENESTATEBINDING_API virtual void ForEachMutableBinding(TFunctionRef<void(FPropertyBindingBinding& Binding)> InFunction) override;
-	SCENESTATEBINDING_API virtual void VisitBindings(TFunctionRef<EVisitResult(const FPropertyBindingBinding& Binding)> InFunction) const override;
-	SCENESTATEBINDING_API virtual void VisitMutableBindings(TFunctionRef<EVisitResult(FPropertyBindingBinding& Binding)> InFunction) override;
-	SCENESTATEBINDING_API virtual void OnReset() override;
-	SCENESTATEBINDING_API virtual void VisitSourceStructDescriptorInternal(TFunctionRef<EVisitResult(const FPropertyBindingBindableStructDescriptor&)> InFunction) const override;
-	[[nodiscard]] SCENESTATEBINDING_API virtual bool OnResolvingPaths() override;
+	UE_API virtual int32 GetNumBindings() const override;
+	UE_API virtual int32 GetNumBindableStructDescriptors() const override;
+	UE_API virtual const FPropertyBindingBindableStructDescriptor* GetBindableStructDescriptorFromHandle(FConstStructView InSourceHandleView) const override;
+	UE_API virtual void ForEachBinding(TFunctionRef<void(const FPropertyBindingBinding& Binding)> InFunction) const override;
+	UE_API virtual void ForEachBinding(FPropertyBindingIndex16 InBegin, FPropertyBindingIndex16 InEnd, TFunctionRef<void(const FPropertyBindingBinding& Binding, const int32 BindingIndex)> InFunction) const override;
+	UE_API virtual void ForEachMutableBinding(TFunctionRef<void(FPropertyBindingBinding& Binding)> InFunction) override;
+	UE_API virtual void VisitBindings(TFunctionRef<EVisitResult(const FPropertyBindingBinding& Binding)> InFunction) const override;
+	UE_API virtual void VisitMutableBindings(TFunctionRef<EVisitResult(FPropertyBindingBinding& Binding)> InFunction) override;
+	UE_API virtual void OnReset() override;
+	UE_API virtual void VisitSourceStructDescriptorInternal(TFunctionRef<EVisitResult(const FPropertyBindingBindableStructDescriptor&)> InFunction) const override;
+	[[nodiscard]] UE_API virtual bool OnResolvingPaths() override;
 	//~ End FPropertyBindingBindingCollection
 
 private:
@@ -85,3 +102,5 @@ private:
 
 	friend UE::SceneState::Editor::FBindingCompiler;
 };
+
+#undef UE_API
diff --git a/Engine/Plugins/VirtualProduction/SceneState/Source/SceneStateBinding/Public/SceneStateBindingCollectionOwner.h b/Engine/Plugins/VirtualProduction/SceneState/Source/SceneStateBinding/Public/SceneStateBindingCollectionOwner.h
new file mode 100644
--- /dev/null
+++ b/Engine/Plugins/VirtualProduction/SceneState/Source/SceneStateBinding/Public/SceneStateBindingCollectionOwner.h
@@ -0,0 +1,42 @@
+// Copyright Epic Games, Inc. All Rights Reserved.
+
+#pragma once
+
+#include "PropertyBindingBindingCollectionOwner.h"
+#include "UObject/Interface.h"
+#include "SceneStateBindingCollectionOwner.generated.h"
+
+struct FSceneStateBindingCollection;
+struct FSceneStateBindingDesc;
+struct FSceneStateBindingFunction;
+
+namespace UE::SceneState
+{
+	struct FBindingFunctionInfo;
+}
+
+UINTERFACE(MinimalAPI, meta=(CannotImplementInterfaceInBlueprint))
+class USceneStateBindingCollectionOwner : public UPropertyBindingBindingCollectionOwner
+{
+	GENERATED_BODY()
+};
+
+class ISceneStateBindingCollectionOwner : public IPropertyBindingBindingCollectionOwner
+{
+	GENERATED_BODY()
+
+public:
+	/** Gets a reference to the binding collection (mutable) that the owner has */
+	virtual FSceneStateBindingCollection& GetBindingCollection() = 0;
+
+	/** Gets a reference to the binding collection that the owner has */
+	virtual const FSceneStateBindingCollection& GetBindingCollection() const = 0;
+
+#if WITH_EDITOR
+	/** Iterates all functions that can be created and bound */
+	virtual bool ForEachBindableFunction(TFunctionRef<bool(const FSceneStateBindingDesc&, const UE::SceneState::FBindingFunctionInfo&)> InFunc) const = 0;
+
+	/** Iterates all functions found within the bindings */
+	virtual bool ForEachBindingFunction(TFunctionRef<bool(const FSceneStateBindingFunction&)> InFunc) const = 0;
+#endif
+};
diff --git a/Engine/Plugins/VirtualProduction/SceneState/Source/SceneStateBinding/Public/SceneStateBindingDataHandle.h b/Engine/Plugins/VirtualProduction/SceneState/Source/SceneStateBinding/Public/SceneStateBindingDataHandle.h
--- a/Engine/Plugins/VirtualProduction/SceneState/Source/SceneStateBinding/Public/SceneStateBindingDataHandle.h
+++ b/Engine/Plugins/VirtualProduction/SceneState/Source/SceneStateBinding/Public/SceneStateBindingDataHandle.h
@@ -13,6 +13,7 @@ enum class ESceneStateDataType : uint8
 	TaskExtension,
 	StateMachine,
 	Transition,
+	Function,
 };
 
 USTRUCT()
@@ -52,6 +53,11 @@ struct FSceneStateBindingDataHandle
 		return Type > std::numeric_limits<uint8>::max();
 	}
 
+	bool IsDataType(ESceneStateDataType InDataType) const
+	{
+		return Type == static_cast<uint16>(InDataType); 
+	}
+
 	uint8 GetDataType() const
 	{
 		return IsExternalDataType()
diff --git a/Engine/Plugins/VirtualProduction/SceneState/Source/SceneStateBinding/Public/SceneStateBindingFunction.h b/Engine/Plugins/VirtualProduction/SceneState/Source/SceneStateBinding/Public/SceneStateBindingFunction.h
new file mode 100644
--- /dev/null
+++ b/Engine/Plugins/VirtualProduction/SceneState/Source/SceneStateBinding/Public/SceneStateBindingFunction.h
@@ -0,0 +1,55 @@
+// Copyright Epic Games, Inc. All Rights Reserved.
+
+#pragma once
+
+#include "StructUtils/InstancedStruct.h"
+#include "SceneStateBindingFunction.generated.h"
+
+#define UE_API SCENESTATEBINDING_API
+
+namespace UE::SceneState
+{
+
+#if WITH_EDITOR
+/** Arguments to create a valid FSceneStateBindingFunction */
+struct FBindingFunctionInfo
+{
+	const UScriptStruct* FunctionType = nullptr;
+	const UScriptStruct* InstanceType = nullptr;
+};
+#endif
+
+} // UE::SceneState
+
+/** Holds a function object and its template instance data for compilation. Used in editor only */
+USTRUCT()
+struct FSceneStateBindingFunction
+{
+	GENERATED_BODY()
+
+	FSceneStateBindingFunction() = default;
+
+#if WITH_EDITOR
+	/** Initializes the binding function with the given function info */
+	UE_API explicit FSceneStateBindingFunction(const UE::SceneState::FBindingFunctionInfo& InFunctionInfo);
+
+	/** Returns whether this object has a valid function and function instance */
+	UE_API bool IsValid() const;
+#endif
+
+#if WITH_EDITORONLY_DATA
+	/** Id of the function */
+	UPROPERTY(VisibleAnywhere, Category="Scene State")
+	FGuid FunctionId;
+
+	/** The function template */
+	UPROPERTY(EditAnywhere, Category="Scene State")
+	FInstancedStruct Function;
+
+	/** The function instance template */
+	UPROPERTY(EditAnywhere, Category="Scene State")
+	mutable FInstancedStruct FunctionInstance;
+#endif
+};
+
+#undef UE_API
diff --git a/Engine/Plugins/VirtualProduction/SceneState/Source/SceneStateBlueprint/Private/SceneStateBlueprint.cpp b/Engine/Plugins/VirtualProduction/SceneState/Source/SceneStateBlueprint/Private/SceneStateBlueprint.cpp
--- a/Engine/Plugins/VirtualProduction/SceneState/Source/SceneStateBlueprint/Private/SceneStateBlueprint.cpp
+++ b/Engine/Plugins/VirtualProduction/SceneState/Source/SceneStateBlueprint/Private/SceneStateBlueprint.cpp
@@ -4,12 +4,14 @@
 #include "Blueprint/BlueprintExtension.h"
 #include "EdGraph/EdGraph.h"
 #include "EdGraph/EdGraphNode.h"
+#include "Functions/SceneStateFunction.h"
 #include "Kismet2/BlueprintEditorUtils.h"
 #include "Modules/ModuleManager.h"
 #include "Nodes/SceneStateMachineStateNode.h"
 #include "Nodes/SceneStateMachineTaskNode.h"
 #include "PropertyBindingDataView.h"
 #include "SceneStateBindingDesc.h"
+#include "SceneStateBindingFunction.h"
 #include "SceneStateBindingUtils.h"
 #include "SceneStateBlueprintBindingUtils.h"
 #include "SceneStateBlueprintDelegates.h"
@@ -17,6 +19,7 @@
 #include "SceneStateBlueprintUtils.h"
 #include "SceneStateGeneratedClass.h"
 #include "SceneStateMachineGraph.h"
+#include "UObject/UObjectIterator.h"
 
 USceneStateBlueprint::USceneStateBlueprint(const FObjectInitializer& InObjectInitializer)
 	: Super(InObjectInitializer)
@@ -56,6 +59,69 @@ FSceneStateBindingDesc USceneStateBlueprint::CreateRootBinding() const
 	return CreateBindingDesc(*this);
 }
 
+FSceneStateBindingCollection& USceneStateBlueprint::GetBindingCollection()
+{
+	return BindingCollection;
+}
+
+const FSceneStateBindingCollection& USceneStateBlueprint::GetBindingCollection() const
+{
+	return BindingCollection;
+}
+
+bool USceneStateBlueprint::ForEachBindableFunction(TFunctionRef<bool(const FSceneStateBindingDesc&, const UE::SceneState::FBindingFunctionInfo&)> InFunc) const
+{
+	const FName MD_Hidden = TEXT("Hidden");
+	const FName MD_Category = TEXT("Category");
+
+	const FString BaseCategory = TEXT("Functions|");
+
+	// Todo: cache functions for faster iteration
+	for (const UScriptStruct* Struct : TObjectRange<UScriptStruct>())
+	{
+		if (!Struct->IsChildOf<FSceneStateFunction>() || Struct->HasMetaData(MD_Hidden))
+		{
+			continue;
+		}
+
+		const UScriptStruct* InstanceType = TInstancedStruct<FSceneStateFunction>(Struct).Get().GetFunctionDataType();
+
+		FSceneStateBindingDesc BindingDesc;
+		BindingDesc.Struct = InstanceType;
+		BindingDesc.Name = *Struct->GetDisplayNameText().ToString();
+		BindingDesc.ID = FGuid::NewDeterministicGuid(Struct->GetStructPathName().ToString());
+		BindingDesc.Category = BaseCategory + Struct->GetMetaData(MD_Category);
+		BindingDesc.DataHandle = FSceneStateBindingDataHandle(ESceneStateDataType::Function);
+
+		UE::SceneState::FBindingFunctionInfo FunctionInfo;
+		FunctionInfo.FunctionType = Struct;
+		FunctionInfo.InstanceType = InstanceType;
+
+		if (!InFunc(BindingDesc, FunctionInfo))
+		{
+			return false;
+		}
+	}
+
+	return true;
+}
+
+bool USceneStateBlueprint::ForEachBindingFunction(TFunctionRef<bool(const FSceneStateBindingFunction&)> InFunc) const
+{
+	for (const FPropertyBindingBinding& Binding : BindingCollection.GetBindings())
+	{
+		const FSceneStateBindingFunction* BindingFunction = Binding.GetPropertyFunctionNode().GetPtr<const FSceneStateBindingFunction>();
+		if (BindingFunction && BindingFunction->IsValid())
+		{
+			if (!InFunc(*BindingFunction))
+			{
+				return false;
+			}
+		}
+	}
+	return true;
+}
+
 void USceneStateBlueprint::GetBindableStructs(const FGuid InTargetStructId, TArray<TInstancedStruct<FPropertyBindingBindableStructDescriptor>>& OutStructDescs) const
 {
 	using namespace UE::SceneState::Graph;
@@ -64,6 +130,7 @@ void USceneStateBlueprint::GetBindableStructs(const FGuid InTargetStructId, TArr
 	if (USceneStateMachineGraph* SceneStateMachine = FindStateMachineMatchingId(*this, InTargetStructId))
 	{
 		TArray<TInstancedStruct<FSceneStateBindingDesc>> BindingDescs;
+		GetGlobalBindingDescs(*this, InTargetStructId, BindingDescs);
 		GetStateMachineBindingDescs(*this, *SceneStateMachine, BindingDescs);
 		OutStructDescs.Append(MoveTemp(BindingDescs));
 		return;
@@ -73,6 +140,7 @@ void USceneStateBlueprint::GetBindableStructs(const FGuid InTargetStructId, TArr
 	if (USceneStateMachineTransitionNode* TransitionNode = FindTransitionMatchingId(*this, InTargetStructId))
 	{
 		TArray<TInstancedStruct<FSceneStateBindingDesc>> BindingDescs;
+		GetGlobalBindingDescs(*this, InTargetStructId, BindingDescs);
 		GetTransitionBindingDescs(*this, *TransitionNode, BindingDescs);
 		OutStructDescs.Append(MoveTemp(BindingDescs));
 		return;
@@ -85,8 +153,18 @@ void USceneStateBlueprint::GetBindableStructs(const FGuid InTargetStructId, TArr
 	if (const USceneStateMachineTaskNode* TargetTaskNode = FindTaskNodeContainingId(*this, InTargetStructId))
 	{
 		TArray<TInstancedStruct<FSceneStateBindingDesc>> BindingDescs;
+		GetGlobalBindingDescs(*this, InTargetStructId, BindingDescs);
 		GetTaskBindingDescs(*this, *TargetTaskNode, BindingDescs);
 		OutStructDescs.Append(MoveTemp(BindingDescs));
+		return;
+	}
+
+	// Add global bindable structs if the target struct is a binding function
+	if (const FSceneStateBindingFunction* BindingFunction = FindBindingFunctionMatchingId(*this, InTargetStructId))
+	{
+		TArray<TInstancedStruct<FSceneStateBindingDesc>> BindingDescs;
+		GetGlobalBindingDescs(*this, InTargetStructId, BindingDescs);
+		OutStructDescs.Append(MoveTemp(BindingDescs));
 	}
 }
 
@@ -143,6 +221,16 @@ void USceneStateBlueprint::CreateParametersForStruct(const FGuid InStructId, TAr
 	}
 }
 
+void USceneStateBlueprint::AppendBindablePropertyFunctionStructs(TArray<TInstancedStruct<FPropertyBindingBindableStructDescriptor>>& InOutStructs) const
+{
+	ForEachBindableFunction(
+		[&InOutStructs](const FSceneStateBindingDesc& InBindingDesc, const UE::SceneState::FBindingFunctionInfo& InFunctionInfo)
+		{
+			InOutStructs.Add(TInstancedStruct<FSceneStateBindingDesc>::Make(InBindingDesc));
+			return true; // continue
+		});
+}
+
 void USceneStateBlueprint::SetObjectBeingDebugged(UObject* InNewObject)
 {
 	using namespace UE::SceneState;
diff --git a/Engine/Plugins/VirtualProduction/SceneState/Source/SceneStateBlueprint/Private/SceneStateBlueprintBindingUtils.cpp b/Engine/Plugins/VirtualProduction/SceneState/Source/SceneStateBlueprint/Private/SceneStateBlueprintBindingUtils.cpp
--- a/Engine/Plugins/VirtualProduction/SceneState/Source/SceneStateBlueprint/Private/SceneStateBlueprintBindingUtils.cpp
+++ b/Engine/Plugins/VirtualProduction/SceneState/Source/SceneStateBlueprint/Private/SceneStateBlueprintBindingUtils.cpp
@@ -6,6 +6,7 @@
 #include "Nodes/SceneStateMachineTransitionNode.h"
 #include "PropertyBindingDataView.h"
 #include "SceneStateBindingDesc.h"
+#include "SceneStateBindingFunction.h"
 #include "SceneStateBlueprint.h"
 #include "SceneStateBlueprint.h"
 #include "SceneStateBlueprintUtils.h"
@@ -86,6 +87,34 @@ FSceneStateBindingDesc CreateBindingDesc(const USceneStateMachineTransitionNode&
 	return BindingDesc;
 }
 
+FSceneStateBindingDesc CreateBindingDesc(const FSceneStateBindingFunction& InBindingFunctionChecked)
+{
+	FSceneStateBindingDesc BindingDesc;
+	BindingDesc.Struct = InBindingFunctionChecked.FunctionInstance.GetScriptStruct();
+	BindingDesc.Name = *InBindingFunctionChecked.Function.GetScriptStruct()->GetDisplayNameText().ToString();
+	BindingDesc.ID = InBindingFunctionChecked.FunctionId;
+	return BindingDesc;
+}
+
+void GetGlobalBindingDescs(const USceneStateBlueprint& InBlueprint, const FGuid& InTargetStructId, TArray<TInstancedStruct<FSceneStateBindingDesc>>& OutBindingDescs)
+{
+	OutBindingDescs.Add(TInstancedStruct<FSceneStateBindingDesc>::Make(CreateBindingDesc(InBlueprint)));
+
+	for (const FPropertyBindingBinding& Binding : InBlueprint.GetBindingCollection().GetBindings())
+	{
+		if (Binding.GetTargetPath().GetStructID() != InTargetStructId)
+		{
+			continue;
+		}
+
+		const FSceneStateBindingFunction* BindingFunction = Binding.GetPropertyFunctionNode().GetPtr<const FSceneStateBindingFunction>();
+		if (BindingFunction && BindingFunction->IsValid())
+		{
+			OutBindingDescs.Add(TInstancedStruct<FSceneStateBindingDesc>::Make(CreateBindingDesc(*BindingFunction)));
+		}
+	}
+}
+
 USceneStateMachineGraph* FindStateMachineMatchingId(const USceneStateBlueprint& InBlueprint, const FGuid& InStructId)
 {
 	USceneStateMachineGraph* FoundGraph = nullptr;
@@ -105,8 +134,6 @@ USceneStateMachineGraph* FindStateMachineMatchingId(const USceneStateBlueprint&
 
 void GetStateMachineBindingDescs(const USceneStateBlueprint& InBlueprint, const USceneStateMachineGraph& InGraph, TArray<TInstancedStruct<FSceneStateBindingDesc>>& OutBindingDescs)
 {
-	OutBindingDescs.Add(TInstancedStruct<FSceneStateBindingDesc>::Make(CreateBindingDesc(InBlueprint)));
-
 	// If this State Machine is under a parent state node, add the event handlers of the state (recursively)
 	USceneStateMachineStateNode* ParentStateNode = InGraph.GetTypedOuter<USceneStateMachineStateNode>(); 
 	GetStateBindingDescs(ParentStateNode, OutBindingDescs, FString());
@@ -133,8 +160,6 @@ USceneStateMachineTaskNode* FindTaskNodeContainingId(const USceneStateBlueprint&
 
 void GetTaskBindingDescs(const USceneStateBlueprint& InBlueprint, const USceneStateMachineTaskNode& InTaskNode, TArray<TInstancedStruct<FSceneStateBindingDesc>>& OutBindingDescs)
 {
-	OutBindingDescs.Add(TInstancedStruct<FSceneStateBindingDesc>::Make(CreateBindingDesc(InBlueprint)));
-
 	// Gather the Binding Descs starting from the directly connected State node.
 	if (USceneStateMachineStateNode* const StateNode = USceneStateMachineGraphSchema::FindConnectedStateNode(&InTaskNode))
 	{
@@ -152,6 +177,24 @@ void GetTaskBindingDescs(const USceneStateBlueprint& InBlueprint, const USceneSt
 	}
 }
 
+const FSceneStateBindingFunction* FindBindingFunctionMatchingId(const USceneStateBlueprint& InBlueprint, const FGuid& InStructId)
+{
+	const FSceneStateBindingFunction* Result = nullptr;
+
+	InBlueprint.ForEachBindingFunction(
+		[&InStructId, &Result](const FSceneStateBindingFunction& InBindingFunction)
+		{
+			if (InBindingFunction.FunctionId == InStructId)
+			{
+				Result = &InBindingFunction;
+				return false; // break
+			}
+			return true; // continue
+		});
+
+	return Result;
+}
+
 USceneStateMachineTransitionNode* FindTransitionMatchingId(const USceneStateBlueprint& InBlueprint, const FGuid& InStructId)
 {
 	USceneStateMachineTransitionNode* FoundNode = nullptr;
@@ -172,8 +215,6 @@ USceneStateMachineTransitionNode* FindTransitionMatchingId(const USceneStateBlue
 
 void GetTransitionBindingDescs(const USceneStateBlueprint& InBlueprint, const USceneStateMachineTransitionNode& InTransitionNode, TArray<TInstancedStruct<FSceneStateBindingDesc>>& OutBindingDescs)
 {
-	OutBindingDescs.Add(TInstancedStruct<FSceneStateBindingDesc>::Make(CreateBindingDesc(InBlueprint)));
-
 	// Gather the Binding Descs starting from the directly outer state machine.
 	if (USceneStateMachineGraph* StateMachineGraph = Cast<USceneStateMachineGraph>(InTransitionNode.GetGraph()))
 	{
@@ -271,6 +312,18 @@ bool FindBindingDescById(const USceneStateBlueprint& InBlueprint, const FGuid& I
 			}
 		});
 
+	// Early exit if the struct desc has already been found (even if invalid)
+	if (Result != EResult::NotFound)
+	{
+		return Result == EResult::FoundValid;
+	}
+
+	if (const FSceneStateBindingFunction* BindingFunction = FindBindingFunctionMatchingId(InBlueprint, InStructId))
+	{
+		OutBindingDesc = TInstancedStruct<FSceneStateBindingDesc>::Make(CreateBindingDesc(*BindingFunction));
+		return true;
+	}
+
 	return Result == EResult::FoundValid;
 }
 
@@ -341,7 +394,19 @@ bool FindDataViewById(const USceneStateBlueprint& InBlueprint, const FGuid& InSt
 			}
 		});
 
-	return bDataViewFound;
+	// Early exit if data view has already been found
+	if (bDataViewFound)
+	{
+		return true;
+	}
+
+	if (const FSceneStateBindingFunction* BindingFunction = FindBindingFunctionMatchingId(InBlueprint, InStructId))
+	{
+		OutDataView = BindingFunction->FunctionInstance;
+		return true;
+	}
+
+	return false;
 }
 
 } // namespace UE::SceneState::Graph
diff --git a/Engine/Plugins/VirtualProduction/SceneState/Source/SceneStateBlueprint/Private/SceneStateBlueprintBindingUtils.h b/Engine/Plugins/VirtualProduction/SceneState/Source/SceneStateBlueprint/Private/SceneStateBlueprintBindingUtils.h
--- a/Engine/Plugins/VirtualProduction/SceneState/Source/SceneStateBlueprint/Private/SceneStateBlueprintBindingUtils.h
+++ b/Engine/Plugins/VirtualProduction/SceneState/Source/SceneStateBlueprint/Private/SceneStateBlueprintBindingUtils.h
@@ -14,6 +14,7 @@ class USceneStateMachineTransitionNode;
 struct FGuid;
 struct FPropertyBindingDataView;
 struct FSceneStateBindingDesc;
+struct FSceneStateBindingFunction;
 template<typename BaseStructT> struct TInstancedStruct;
 
 namespace UE::SceneState::Graph
@@ -35,6 +36,17 @@ namespace UE::SceneState::Graph
 	/** Creates the Binding Struct Desc for the Parameters in the given Transition Node */
 	FSceneStateBindingDesc CreateBindingDesc(const USceneStateMachineTransitionNode& InTransitionNode);
 
+	/** Creates the Binding Struct Desc for the given binding function */
+	FSceneStateBindingDesc CreateBindingDesc(const FSceneStateBindingFunction& InBindingFunctionChecked);
+
+	/**
+	 * Gathers all the global binding descs. Here 'global' means within the scope of the blueprint.
+	 * @param InBlueprint the blueprint
+	 * @param InTargetStructId the struct to gather the binding descs for
+	 * @param OutBindingDescs all the binding descs that can be bound for the target struct
+	 */
+	void GetGlobalBindingDescs(const USceneStateBlueprint& InBlueprint, const FGuid& InTargetStructId, TArray<TInstancedStruct<FSceneStateBindingDesc>>& OutBindingDescs);
+
 	/**
 	 * Finds the state machine graph with id that matches the given struct id.
 	 * @param InBlueprint the blueprint to look into
@@ -67,6 +79,14 @@ namespace UE::SceneState::Graph
 	 */
 	void GetTaskBindingDescs(const USceneStateBlueprint& InBlueprint, const USceneStateMachineTaskNode& InTaskNode, TArray<TInstancedStruct<FSceneStateBindingDesc>>& OutBindingDescs);
 
+	/**
+	 * Finds the binding function that matches the given struct id.
+	 * @param InBlueprint the blueprint to look into
+	 * @param InStructId the struct id to look for
+	 * @return the found binding function containing the given struct id
+	 */
+	const FSceneStateBindingFunction* FindBindingFunctionMatchingId(const USceneStateBlueprint& InBlueprint, const FGuid& InStructId);
+
 	/**
 	 * Finds the transition node with id that matches the given struct id.
 	 * @param InBlueprint the blueprint to look into
diff --git a/Engine/Plugins/VirtualProduction/SceneState/Source/SceneStateBlueprint/Public/SceneStateBlueprint.h b/Engine/Plugins/VirtualProduction/SceneState/Source/SceneStateBlueprint/Public/SceneStateBlueprint.h
--- a/Engine/Plugins/VirtualProduction/SceneState/Source/SceneStateBlueprint/Public/SceneStateBlueprint.h
+++ b/Engine/Plugins/VirtualProduction/SceneState/Source/SceneStateBlueprint/Public/SceneStateBlueprint.h
@@ -4,21 +4,23 @@
 
 #include "Containers/Array.h"
 #include "Engine/Blueprint.h"
-#include "PropertyBindingBindingCollectionOwner.h"
 #include "SceneStateBindingCollection.h"
+#include "SceneStateBindingCollectionOwner.h"
 #include "SceneStateBlueprint.generated.h"
 
+#define UE_API SCENESTATEBLUEPRINT_API
+
 class USceneStateGeneratedClass;
 class USceneStateMachineGraph;
 class USceneStateMachineStateNode;
 
 UCLASS(MinimalAPI, DisplayName="Motion Design Scene State Blueprint")
-class USceneStateBlueprint : public UBlueprint, public IPropertyBindingBindingCollectionOwner
+class USceneStateBlueprint : public UBlueprint, public ISceneStateBindingCollectionOwner
 {
 	GENERATED_BODY()
 
 public:
-	SCENESTATEBLUEPRINT_API USceneStateBlueprint(const FObjectInitializer& InObjectInitializer);
+	UE_API USceneStateBlueprint(const FObjectInitializer& InObjectInitializer);
 
 	const FGuid& GetRootId() const
 	{
@@ -31,31 +33,39 @@ public:
 		return Cast<T>(FindExtension(T::StaticClass()));
 	}
 
-	SCENESTATEBLUEPRINT_API UBlueprintExtension* FindExtension(TSubclassOf<UBlueprintExtension> InClass) const;
+	UE_API UBlueprintExtension* FindExtension(TSubclassOf<UBlueprintExtension> InClass) const;
+
+	UE_API FSceneStateBindingDesc CreateRootBinding() const;
 
-	SCENESTATEBLUEPRINT_API FSceneStateBindingDesc CreateRootBinding() const;
+	//~ Begin ISceneStateBindingCollectionOwner
+	UE_API virtual FSceneStateBindingCollection& GetBindingCollection() override;
+	UE_API virtual const FSceneStateBindingCollection& GetBindingCollection() const override;
+	UE_API virtual bool ForEachBindableFunction(TFunctionRef<bool(const FSceneStateBindingDesc&, const UE::SceneState::FBindingFunctionInfo&)> InFunc) const override;
+	UE_API virtual bool ForEachBindingFunction(TFunctionRef<bool(const FSceneStateBindingFunction&)> InFunc) const;
+	//~ End ISceneStateBindingCollectionOwner
 
 	//~ Begin IPropertyBindingBindingCollectionOwner
-	SCENESTATEBLUEPRINT_API virtual void GetBindableStructs(const FGuid InTargetStructId, TArray<TInstancedStruct<FPropertyBindingBindableStructDescriptor>>& OutStructDescs) const override;
-	SCENESTATEBLUEPRINT_API virtual bool GetBindableStructByID(const FGuid InStructId, TInstancedStruct<FPropertyBindingBindableStructDescriptor>& OutStructDesc) const override;
-	SCENESTATEBLUEPRINT_API virtual bool GetBindingDataViewByID(const FGuid InStructId, FPropertyBindingDataView& OutDataView) const override;
-	SCENESTATEBLUEPRINT_API virtual FPropertyBindingBindingCollection* GetEditorPropertyBindings() override;
-	SCENESTATEBLUEPRINT_API virtual const FPropertyBindingBindingCollection* GetEditorPropertyBindings() const override;
-	SCENESTATEBLUEPRINT_API virtual bool CanCreateParameter(const FGuid InStructId) const override;
-	SCENESTATEBLUEPRINT_API virtual void CreateParametersForStruct(const FGuid InStructId, TArrayView<UE::PropertyBinding::FPropertyCreationDescriptor> InOutCreationDescs) override;
+	UE_API virtual void GetBindableStructs(const FGuid InTargetStructId, TArray<TInstancedStruct<FPropertyBindingBindableStructDescriptor>>& OutStructDescs) const override;
+	UE_API virtual bool GetBindableStructByID(const FGuid InStructId, TInstancedStruct<FPropertyBindingBindableStructDescriptor>& OutStructDesc) const override;
+	UE_API virtual bool GetBindingDataViewByID(const FGuid InStructId, FPropertyBindingDataView& OutDataView) const override;
+	UE_API virtual FPropertyBindingBindingCollection* GetEditorPropertyBindings() override;
+	UE_API virtual const FPropertyBindingBindingCollection* GetEditorPropertyBindings() const override;
+	UE_API virtual bool CanCreateParameter(const FGuid InStructId) const override;
+	UE_API virtual void CreateParametersForStruct(const FGuid InStructId, TArrayView<UE::PropertyBinding::FPropertyCreationDescriptor> InOutCreationDescs) override;
+	UE_API virtual void AppendBindablePropertyFunctionStructs(TArray<TInstancedStruct<FPropertyBindingBindableStructDescriptor>>& InOutStructs) const override;
 	//~ End IPropertyBindingBindingCollectionOwner
 
 	//~ Begin UBlueprint
-	SCENESTATEBLUEPRINT_API virtual void SetObjectBeingDebugged(UObject* InNewObject) override;
-	SCENESTATEBLUEPRINT_API virtual UClass* GetBlueprintClass() const override;
-	SCENESTATEBLUEPRINT_API virtual void GetReparentingRules(TSet<const UClass*>& OutAllowedChildrenOfClasses, TSet<const UClass*>& OutDisallowedChildrenOfClasses) const override;
-	SCENESTATEBLUEPRINT_API virtual bool SupportedByDefaultBlueprintFactory() const override;
-	SCENESTATEBLUEPRINT_API virtual void LoadModulesRequiredForCompilation() override;
-	SCENESTATEBLUEPRINT_API virtual bool IsValidForBytecodeOnlyRecompile() const override;
+	UE_API virtual void SetObjectBeingDebugged(UObject* InNewObject) override;
+	UE_API virtual UClass* GetBlueprintClass() const override;
+	UE_API virtual void GetReparentingRules(TSet<const UClass*>& OutAllowedChildrenOfClasses, TSet<const UClass*>& OutDisallowedChildrenOfClasses) const override;
+	UE_API virtual bool SupportedByDefaultBlueprintFactory() const override;
+	UE_API virtual void LoadModulesRequiredForCompilation() override;
+	UE_API virtual bool IsValidForBytecodeOnlyRecompile() const override;
 	//~ End UBlueprint
 
 	//~ Begin UObject
-	SCENESTATEBLUEPRINT_API virtual void BeginDestroy() override;
+	UE_API virtual void BeginDestroy() override;
 	//~ End UObject
 
 	/** The top level State Machine Graphs of this State. Does not include the state machine graphs of nested state nodes */
@@ -89,3 +99,5 @@ private:
 	/** Handle to the delegate when a state machine graph parameters have changed */
 	FDelegateHandle OnGraphParametersChangedHandle;
 };
+
+#undef UE_API
diff --git a/Engine/Plugins/VirtualProduction/SceneState/Source/SceneStateBlueprintEditor/Private/Compiler/SceneStateBindingCompiler.cpp b/Engine/Plugins/VirtualProduction/SceneState/Source/SceneStateBlueprintEditor/Private/Compiler/SceneStateBindingCompiler.cpp
--- a/Engine/Plugins/VirtualProduction/SceneState/Source/SceneStateBlueprintEditor/Private/Compiler/SceneStateBindingCompiler.cpp
+++ b/Engine/Plugins/VirtualProduction/SceneState/Source/SceneStateBlueprintEditor/Private/Compiler/SceneStateBindingCompiler.cpp
@@ -1,7 +1,9 @@
 // Copyright Epic Games, Inc. All Rights Reserved.
 
 #include "SceneStateBindingCompiler.h"
+#include "Functions/SceneStateFunction.h"
 #include "Misc/EnumerateRange.h"
+#include "SceneStateBindingFunction.h"
 #include "SceneStateBindingUtils.h"
 #include "SceneStateBlueprint.h"
 #include "SceneStateBlueprintCompilerContext.h"
@@ -57,6 +59,14 @@ EDataAccessType GetBindingBatchingType(const FSceneStateBindingCollection& InBin
 
 } // UE::SceneState::Editor::Private
 
+FBindingCompiler::FBindingCompiler(FBlueprintCompilerContext& InContext, USceneStateBlueprint* InBlueprint, USceneStateGeneratedClass* InGeneratedClass)
+	: Context(InContext)
+	, Blueprint(InBlueprint)
+	, Class(InGeneratedClass)
+	, BindingFunctionCompiler(InGeneratedClass)
+{
+}
+
 void FBindingCompiler::Compile()
 {
 	check(Blueprint && Class);
@@ -70,11 +80,14 @@ void FBindingCompiler::Compile()
 
 	ValidBindingMap.Reserve(BindingDescInitialCount);
 
+	BindingFunctionCompiler.Compile();
+
 	AddRootBindingDesc();
 	AddStateMachineBindingDescs();
 	AddTransitionBindingDescs();
 	AddTaskBindingDescs();
 	AddEventHandlerBindingDescs();
+	AddFunctionBindingDescs();
 
 	RemoveInvalidBindings();
 
@@ -222,6 +235,26 @@ void FBindingCompiler::AddEventHandlerBindingDescs()
 	}
 }
 
+void FBindingCompiler::AddFunctionBindingDescs()
+{
+	check(Class->Functions.Num() == Class->FunctionMetadata.Num()
+		&& Class->Functions.Num() == Class->FunctionInstances.Num());
+
+	for (int32 FunctionIndex = 0; FunctionIndex < Class->Functions.Num(); ++FunctionIndex)
+	{
+		const FStructView FunctionInstance = Class->FunctionInstances[FunctionIndex];
+		const FSceneStateFunctionMetadata& FunctionMetadata = Class->FunctionMetadata[FunctionIndex];
+
+		FSceneStateBindingDesc BindingDesc;
+		BindingDesc.ID = FunctionMetadata.FunctionId;
+		BindingDesc.Name = FName(TEXT("Function"), FunctionIndex);
+		BindingDesc.Struct = FunctionInstance.GetScriptStruct();
+		BindingDesc.DataHandle = FSceneStateBindingDataHandle(ESceneStateDataType::Function, FunctionIndex);
+
+		AddBindingDesc(MoveTemp(BindingDesc));
+	}
+}
+
 bool FBindingCompiler::ValidateBinding(const FSceneStateBinding& InBinding) const
 {
 	// Source must be accessible by the target struct.
@@ -373,6 +406,15 @@ void FBindingCompiler::CompileCopies()
 
 		// Binding end is one past the last binding.
 		CopyBatch.BindingsEnd = FPropertyBindingIndex16(BindingIndex);
+
+		// Get the compiled function range for this struct
+		const FSceneStateRange FunctionRange = BindingFunctionCompiler.GetFunctionRange(TargetDesc->ID);
+		if (FunctionRange.IsValid())
+		{
+			CopyBatch.PropertyFunctionsBegin = FPropertyBindingIndex16(FunctionRange.Index);
+			CopyBatch.PropertyFunctionsEnd = FPropertyBindingIndex16(FunctionRange.Index + FunctionRange.Count);
+		}
+
 		OnBindingsBatchCompiled(BatchBindings, CopyBinding.TargetDataHandle);
 	}
 }
@@ -469,6 +511,11 @@ void FBindingCompiler::OnBindingsBatchCompiled(FPropertyBindingIndex16 InBinding
 		FSceneStateTransition& Transition = Class->Transitions[DataIndex];
 		Transition.BindingsBatch = InBindingsBatch;
 	}
+	else if (DataType == ESceneStateDataType::Function)
+	{
+		FSceneStateFunction& Function = Class->Functions[DataIndex].Get<FSceneStateFunction>();
+		Function.BindingsBatch = InBindingsBatch;
+	}
 	else
 	{
 		// No other Data Types are supported as Targets at the moment
@@ -540,6 +587,16 @@ FSceneStateBindingDataHandle FBindingCompiler::GetDataHandleById(const FGuid& In
 		}
 	}
 
+	// Functions
+	for (int32 FunctionIndex = 0; FunctionIndex < Class->Functions.Num(); ++FunctionIndex)
+	{
+		const FSceneStateFunctionMetadata& FunctionMetadata = Class->FunctionMetadata[FunctionIndex];
+		if (FunctionMetadata.FunctionId == InStructId)
+		{
+			return FSceneStateBindingDataHandle(ESceneStateDataType::Function, FunctionIndex);
+		}
+	}
+
 	return FSceneStateBindingDataHandle();
 }
 
diff --git a/Engine/Plugins/VirtualProduction/SceneState/Source/SceneStateBlueprintEditor/Private/Compiler/SceneStateBindingCompiler.h b/Engine/Plugins/VirtualProduction/SceneState/Source/SceneStateBlueprintEditor/Private/Compiler/SceneStateBindingCompiler.h
--- a/Engine/Plugins/VirtualProduction/SceneState/Source/SceneStateBlueprintEditor/Private/Compiler/SceneStateBindingCompiler.h
+++ b/Engine/Plugins/VirtualProduction/SceneState/Source/SceneStateBlueprintEditor/Private/Compiler/SceneStateBindingCompiler.h
@@ -5,6 +5,7 @@
 #include "Misc/Guid.h"
 #include "PropertyBindingDataView.h"
 #include "SceneStateBindingDataHandle.h"
+#include "SceneStateBindingFunctionCompiler.h"
 #include "UObject/ObjectPtr.h"
 
 class USceneStateBlueprint;
@@ -32,12 +33,7 @@ enum class EDataAccessType : uint8
 class FBindingCompiler
 {
 public:
-	explicit FBindingCompiler(FBlueprintCompilerContext& InContext, USceneStateBlueprint* InBlueprint, USceneStateGeneratedClass* InGeneratedClass)
-		: Context(InContext)
-		, Blueprint(InBlueprint)
-		, Class(InGeneratedClass)
-	{
-	}
+	explicit FBindingCompiler(FBlueprintCompilerContext& InContext, USceneStateBlueprint* InBlueprint, USceneStateGeneratedClass* InGeneratedClass);
 
 	void Compile();
 
@@ -50,6 +46,7 @@ private:
 	void AddTransitionBindingDescs();
 	void AddTaskBindingDescs();
 	void AddEventHandlerBindingDescs();
+	void AddFunctionBindingDescs();
 
 	bool ValidateBinding(const FSceneStateBinding& InBinding) const;
 
@@ -73,6 +70,9 @@ private:
 	/** The generated class to compile bindings for */
 	TObjectPtr<USceneStateGeneratedClass> Class;
 
+	/** Compiler for the functions found within the bindings */
+	FBindingFunctionCompiler BindingFunctionCompiler;
+
 	/** Map of struct ids to their binding data view */
 	TMap<FGuid, const FPropertyBindingDataView> ValidBindingMap;
 
diff --git a/Engine/Plugins/VirtualProduction/SceneState/Source/SceneStateBlueprintEditor/Private/Compiler/SceneStateBindingFunctionCompiler.cpp b/Engine/Plugins/VirtualProduction/SceneState/Source/SceneStateBlueprintEditor/Private/Compiler/SceneStateBindingFunctionCompiler.cpp
new file mode 100644
--- /dev/null
+++ b/Engine/Plugins/VirtualProduction/SceneState/Source/SceneStateBlueprintEditor/Private/Compiler/SceneStateBindingFunctionCompiler.cpp
@@ -0,0 +1,124 @@
+﻿// Copyright Epic Games, Inc. All Rights Reserved.
+
+#include "SceneStateBindingFunctionCompiler.h"
+#include "SceneStateBindingFunction.h"
+#include "SceneStateGeneratedClass.h"
+#include "Tasks/SceneStateTaskBindingExtension.h"
+
+namespace UE::SceneState::Editor
+{
+
+void FBindingFunctionCompiler::Compile()
+{
+	FunctionRangeMap.Reset();
+
+	CompileStateMachineFunctions();
+	CompileTransitionFunctions();
+	CompileTaskFunctions();
+}
+
+FSceneStateRange FBindingFunctionCompiler::GetFunctionRange(const FGuid& InTargetStructId) const
+{
+	return FunctionRangeMap.FindRef(InTargetStructId);
+}
+
+void FBindingFunctionCompiler::CompileFunctions(const FGuid& InTargetStructId)
+{
+	if (!InTargetStructId.IsValid())
+	{
+		return;
+	}
+
+	// If this happens, function bindings are circular.
+	if (!ensure(!FunctionRangeMap.Contains(InTargetStructId)))
+	{
+		return;
+	}
+	FunctionRangeMap.Add(InTargetStructId);
+
+	const FSceneStateBindingCollection& BindingCollection = Class->GetBindingCollection();
+
+	TArray<FConstStructView> Functions;
+	TArray<FConstStructView> FunctionInstances;
+	TArray<FSceneStateFunctionMetadata> FunctionMetadata;
+
+	BindingCollection.ForEachBinding(
+		[This=this, &Functions, &FunctionInstances, &FunctionMetadata, &InTargetStructId](const FPropertyBindingBinding& InBinding)
+		{
+			// Skip bindings that don't have a valid source or don't match the target struct id
+			if (!InBinding.GetSourcePath().GetStructID().IsValid() || InBinding.GetTargetPath().GetStructID() != InTargetStructId)
+			{
+				return;
+			}
+
+			if (const FSceneStateBindingFunction* BindingFunction = InBinding.GetPropertyFunctionNode().GetPtr<const FSceneStateBindingFunction>())
+			{
+				Functions.Add(BindingFunction->Function);
+				FunctionInstances.Add(BindingFunction->FunctionInstance);
+
+				FSceneStateFunctionMetadata& Metadata = FunctionMetadata.AddDefaulted_GetRef();
+				Metadata.FunctionId = BindingFunction->FunctionId;
+
+				// Compile the functions that target this binding function
+				This->CompileFunctions(BindingFunction->FunctionId);
+			}
+		});
+
+	check(Functions.Num() == FunctionInstances.Num() && Functions.Num() == FunctionMetadata.Num());
+
+	if (Functions.IsEmpty())
+	{
+		FunctionRangeMap.Remove(InTargetStructId);
+		return;
+	}
+
+	FSceneStateRange& FunctionRange = FunctionRangeMap[InTargetStructId];
+	FunctionRange.Index = Class->Functions.Num();
+	FunctionRange.Count = Functions.Num();
+
+	Class->Functions.Append(Functions);
+	Class->FunctionInstances.Append(FunctionInstances);
+	Class->FunctionMetadata.Append(MoveTemp(FunctionMetadata));
+}
+
+void FBindingFunctionCompiler::CompileStateMachineFunctions()
+{
+	for (const TPair<FGuid, uint16>& Pair : Class->StateMachineIdToIndex)
+	{
+		CompileFunctions(Pair.Key);
+	}
+}
+
+void FBindingFunctionCompiler::CompileTransitionFunctions()
+{
+	for (const TPair<uint16, FInstancedPropertyBag>& Pair : Class->TransitionParameters)
+	{
+		const FSceneStateTransitionMetadata& TransitionMetadata = Class->TransitionMetadata[Pair.Key];
+		CompileFunctions(TransitionMetadata.ParametersId);
+	}
+}
+
+void FBindingFunctionCompiler::CompileTaskFunctions()
+{
+	check(Class->Tasks.Num() == Class->GetTaskMetadata().Num());
+
+	for (int32 TaskIndex = 0; TaskIndex < Class->Tasks.Num(); ++TaskIndex)
+	{
+		const FSceneStateTaskMetadata& TaskMetadata = Class->TaskMetadata[TaskIndex];
+		CompileFunctions(TaskMetadata.TaskId);
+
+		const FSceneStateTask& Task = Class->Tasks[TaskIndex].Get<const FSceneStateTask>();
+		const FStructView TaskInstance = Class->TaskInstances[TaskIndex];
+
+		if (const FSceneStateTaskBindingExtension* BindingExtension = Task.GetBindingExtension())
+		{
+			BindingExtension->VisitBindingDescs(TaskInstance,
+				[This=this](const FTaskBindingDesc& InBindingDesc)
+				{
+					This->CompileFunctions(InBindingDesc.Id);
+				});
+		}
+	}
+}
+
+} // UE::SceneState::Editor
diff --git a/Engine/Plugins/VirtualProduction/SceneState/Source/SceneStateBlueprintEditor/Private/Compiler/SceneStateBindingFunctionCompiler.h b/Engine/Plugins/VirtualProduction/SceneState/Source/SceneStateBlueprintEditor/Private/Compiler/SceneStateBindingFunctionCompiler.h
new file mode 100644
--- /dev/null
+++ b/Engine/Plugins/VirtualProduction/SceneState/Source/SceneStateBlueprintEditor/Private/Compiler/SceneStateBindingFunctionCompiler.h
@@ -0,0 +1,48 @@
+﻿// Copyright Epic Games, Inc. All Rights Reserved.
+
+#pragma once
+
+#include "SceneStateRange.h"
+#include "UObject/ObjectPtr.h"
+
+class USceneStateGeneratedClass;
+
+namespace UE::SceneState::Editor
+{
+
+/** Compiles all the functions found throughout all the bindings */
+class FBindingFunctionCompiler
+{
+public:
+	explicit FBindingFunctionCompiler(USceneStateGeneratedClass* InGeneratedClass)
+		: Class(InGeneratedClass)
+	{
+	}
+
+	/** Compiles all the functions in all the bindings meant to have functions (Tasks, State Machines, Transitions) */
+	void Compile();
+
+	/** Gets the range of the compiled functions that are bound to the target struct id */
+	FSceneStateRange GetFunctionRange(const FGuid& InTargetStructId) const;
+
+private:
+	/** Compiles the functions found in all the bindings pointing towards target struct id */
+	void CompileFunctions(const FGuid& InTargetStructId);
+
+	/** Compiles the functions for all the compiled tasks */
+	void CompileTaskFunctions();
+
+	/** Compiles the functions for all the compiled state machines */
+	void CompileStateMachineFunctions();
+
+	/** Compiles the functions for all the compiled transitions */
+	void CompileTransitionFunctions();
+
+	/** The generated class to compile functions for */
+	TObjectPtr<USceneStateGeneratedClass> Class;
+
+	/** Map of the target struct id to its range of the functions in the generated class */
+	TMap<FGuid, FSceneStateRange> FunctionRangeMap;
+};
+
+} // UE::SceneState::Editor
diff --git a/Engine/Plugins/VirtualProduction/SceneState/Source/SceneStateBlueprintEditor/Private/DetailsView/SceneStateBindingFunctionCustomization.cpp b/Engine/Plugins/VirtualProduction/SceneState/Source/SceneStateBlueprintEditor/Private/DetailsView/SceneStateBindingFunctionCustomization.cpp
new file mode 100644
--- /dev/null
+++ b/Engine/Plugins/VirtualProduction/SceneState/Source/SceneStateBlueprintEditor/Private/DetailsView/SceneStateBindingFunctionCustomization.cpp
@@ -0,0 +1,70 @@
+﻿// Copyright Epic Games, Inc. All Rights Reserved.
+
+#include "SceneStateBindingFunctionCustomization.h"
+#include "IDetailChildrenBuilder.h"
+#include "IDetailPropertyRow.h"
+#include "InstancedStructDetails.h"
+#include "PropertyHandle.h"
+#include "SceneStateBindingFunction.h"
+#include "SceneStateBlueprintEditorUtils.h"
+#include "SceneStateCachedBindingData.h"
+
+namespace UE::SceneState::Editor
+{
+
+namespace Private
+{
+/** Details customization for function instances - hiding the output property and assigning the function id to the rest */
+class FFunctionInstanceDetails : public FInstancedStructDataDetails
+{
+public:
+	FFunctionInstanceDetails(TSharedPtr<IPropertyHandle> InStructProperty, const FGuid& InFunctionId)
+		: FInstancedStructDataDetails(InStructProperty)
+		, FunctionId(InFunctionId)
+	{
+	}
+
+	//~ Begin FInstancedStructDataDetails
+	virtual void OnChildRowAdded(IDetailPropertyRow& InChildRow) override
+	{
+		TSharedPtr<IPropertyHandle> ChildPropHandle = InChildRow.GetPropertyHandle();
+		check(ChildPropHandle.IsValid());
+		if (UE::SceneState::Editor::IsOutputProperty(ChildPropHandle->GetProperty()))
+		{
+			InChildRow.Visibility(EVisibility::Collapsed);
+		}
+		else
+		{
+			AssignBindingId(ChildPropHandle.ToSharedRef(), FunctionId);
+		}
+	}
+	//~ End FInstancedStructDataDetails
+
+private:
+	FGuid FunctionId;
+};
+
+} // Private
+	
+void FBindingFunctionCustomization::CustomizeHeader(TSharedRef<IPropertyHandle> InPropertyHandle, FDetailWidgetRow& InHeaderRow, IPropertyTypeCustomizationUtils& InCustomizationUtils)
+{
+}
+
+void FBindingFunctionCustomization::CustomizeChildren(TSharedRef<IPropertyHandle> InPropertyHandle, IDetailChildrenBuilder& InChildBuilder, IPropertyTypeCustomizationUtils& InCustomizationUtils)
+{
+	const TSharedRef<IPropertyHandle> FunctionIdHandle = InPropertyHandle->GetChildHandle(GET_MEMBER_NAME_CHECKED(FSceneStateBindingFunction, FunctionId)).ToSharedRef();
+	const TSharedRef<IPropertyHandle> FunctionHandle = InPropertyHandle->GetChildHandle(GET_MEMBER_NAME_CHECKED(FSceneStateBindingFunction, Function)).ToSharedRef();
+	const TSharedRef<IPropertyHandle> FunctionInstanceHandle = InPropertyHandle->GetChildHandle(GET_MEMBER_NAME_CHECKED(FSceneStateBindingFunction, FunctionInstance)).ToSharedRef();
+
+	FunctionIdHandle->MarkHiddenByCustomization();
+	FunctionHandle->MarkHiddenByCustomization();
+	FunctionInstanceHandle->MarkHiddenByCustomization();
+
+	FGuid FunctionId;
+	GetGuid(FunctionIdHandle, FunctionId);
+
+	InChildBuilder.AddCustomBuilder(MakeShared<FInstancedStructDataDetails>(FunctionHandle));
+	InChildBuilder.AddCustomBuilder(MakeShared<Private::FFunctionInstanceDetails>(FunctionInstanceHandle, FunctionId));
+}
+
+} // UE::SceneState::Editor
diff --git a/Engine/Plugins/VirtualProduction/SceneState/Source/SceneStateBlueprintEditor/Private/DetailsView/SceneStateBindingFunctionCustomization.h b/Engine/Plugins/VirtualProduction/SceneState/Source/SceneStateBlueprintEditor/Private/DetailsView/SceneStateBindingFunctionCustomization.h
new file mode 100644
--- /dev/null
+++ b/Engine/Plugins/VirtualProduction/SceneState/Source/SceneStateBlueprintEditor/Private/DetailsView/SceneStateBindingFunctionCustomization.h
@@ -0,0 +1,25 @@
+﻿// Copyright Epic Games, Inc. All Rights Reserved.
+
+#pragma once
+
+#include "IPropertyTypeCustomization.h"
+
+namespace UE::SceneState::Editor
+{
+
+/** Customization for FSceneStateBindingFunction */
+class FBindingFunctionCustomization : public IPropertyTypeCustomization
+{
+public:
+	static TSharedRef<IPropertyTypeCustomization> MakeInstance()
+	{
+		return MakeShared<FBindingFunctionCustomization>();
+	}
+
+	//~ Begin IPropertyTypeCustomization
+	virtual void CustomizeHeader(TSharedRef<IPropertyHandle> InPropertyHandle, FDetailWidgetRow& InHeaderRow, IPropertyTypeCustomizationUtils& InCustomizationUtils) override;
+	virtual void CustomizeChildren(TSharedRef<IPropertyHandle> InPropertyHandle, IDetailChildrenBuilder& InChildBuilder, IPropertyTypeCustomizationUtils& InCustomizationUtils) override;
+	//~ End IPropertyTypeCustomization
+};
+
+} // UE::SceneState::Editor
diff --git a/Engine/Plugins/VirtualProduction/SceneState/Source/SceneStateBlueprintEditor/Private/DetailsView/SceneStateBindingsChildrenCustomization.cpp b/Engine/Plugins/VirtualProduction/SceneState/Source/SceneStateBlueprintEditor/Private/DetailsView/SceneStateBindingsChildrenCustomization.cpp
new file mode 100644
--- /dev/null
+++ b/Engine/Plugins/VirtualProduction/SceneState/Source/SceneStateBlueprintEditor/Private/DetailsView/SceneStateBindingsChildrenCustomization.cpp
@@ -0,0 +1,162 @@
+﻿// Copyright Epic Games, Inc. All Rights Reserved.
+
+#include "SceneStateBindingsChildrenCustomization.h"
+#include "IDetailChildrenBuilder.h"
+#include "PropertyBindingExtension.h"
+#include "PropertyBindingPath.h"
+#include "PropertyHandle.h"
+#include "SceneStateBindingCollection.h"
+#include "SceneStateBindingCollectionOwner.h"
+#include "SceneStateBindingFunction.h"
+
+namespace UE::SceneState::Editor
+{
+
+namespace Private
+{
+/** Finds the bindings owner that this property handle is outered to */
+ISceneStateBindingCollectionOwner* GetBindingsOwner(const TSharedPtr<IPropertyHandle>& InPropertyHandle)
+{
+	if (!InPropertyHandle.IsValid())
+	{
+		return nullptr;
+	}
+
+	TArray<UObject*> OuterObjects;
+	InPropertyHandle->GetOuterObjects(OuterObjects);
+
+	auto FindBindingOwner =
+		[](UObject* InObject)->ISceneStateBindingCollectionOwner*
+		{
+			for (UObject* Outer = InObject; Outer; Outer = Outer->GetOuter())
+			{
+				if (ISceneStateBindingCollectionOwner* BindingOwner = Cast<ISceneStateBindingCollectionOwner>(Outer))
+				{
+					return BindingOwner;
+				}
+			}
+			return nullptr;
+		};
+
+	ISceneStateBindingCollectionOwner* CommonBindingsOwner = nullptr;
+
+	for (UObject* Object : OuterObjects)
+	{
+		ISceneStateBindingCollectionOwner* const BindingsOwner = FindBindingOwner(Object);
+		if (CommonBindingsOwner && CommonBindingsOwner != BindingsOwner)
+		{
+			return nullptr;
+		}
+		CommonBindingsOwner = BindingsOwner;
+	}
+
+	return CommonBindingsOwner;
+}
+
+} // UE::SceneState::Editor::Private
+
+FBindingFunctionStructProvider::FBindingFunctionStructProvider(ISceneStateBindingCollectionOwner* InBindingsOwnerChecked, FPropertyBindingPath InTargetPath)
+	: BindingsOwnerWeak(InBindingsOwnerChecked)
+	, TargetPath(MoveTemp(InTargetPath))
+{
+}
+
+bool FBindingFunctionStructProvider::IsValid() const
+{
+	return GetBindingFunctionView(BindingsOwnerWeak.Get(), TargetPath).IsValid();
+}
+
+const UStruct* FBindingFunctionStructProvider::GetBaseStructure() const
+{
+	return FSceneStateBindingFunction::StaticStruct();
+}
+
+void FBindingFunctionStructProvider::GetInstances(TArray<TSharedPtr<FStructOnScope>>& OutInstances, const UStruct* InExpectedBaseStructure) const
+{
+	if (InExpectedBaseStructure)
+	{
+		const FStructView BindingFunctionView = GetBindingFunctionView(BindingsOwnerWeak.Get(), TargetPath);
+
+		if (BindingFunctionView.IsValid() && BindingFunctionView.GetScriptStruct()->IsChildOf(InExpectedBaseStructure))
+		{
+			OutInstances.Add(MakeShared<FStructOnScope>(BindingFunctionView.GetScriptStruct(), BindingFunctionView.GetMemory()));
+		}
+	}
+}
+
+bool FBindingFunctionStructProvider::HasValidBindingFunction(ISceneStateBindingCollectionOwner* InBindingsOwner, const FPropertyBindingPath& InTargetPath)
+{
+	return GetBindingFunctionView(InBindingsOwner, InTargetPath).IsValid();
+}
+
+FStructView FBindingFunctionStructProvider::GetBindingFunctionView(ISceneStateBindingCollectionOwner* InBindingsOwner, const FPropertyBindingPath& InTargetPath)
+{
+	if (!InBindingsOwner || InTargetPath.IsPathEmpty())
+	{
+		return FStructView();
+	}
+
+	FPropertyBindingBinding* const FoundBinding = InBindingsOwner->GetBindingCollection().GetMutableBindings().FindByPredicate(
+		[&InTargetPath](const FSceneStateBinding& InBinding)
+		{
+			return InBinding.GetTargetPath() == InTargetPath;
+		});
+
+	if (!FoundBinding)
+	{
+		return FStructView();
+	}
+
+	const FStructView BindingFunctionView = FoundBinding->GetMutablePropertyFunctionNode();
+
+	const FSceneStateBindingFunction* BindingFunction = BindingFunctionView.GetPtr<FSceneStateBindingFunction>();
+	if (BindingFunction && BindingFunction->IsValid())
+	{
+		return BindingFunctionView;
+	}
+	return FStructView();
+}
+
+bool FBindingsChildrenCustomization::ShouldCustomizeChildren(TSharedRef<IPropertyHandle> InPropertyHandle)
+{
+	const ISceneStateBindingCollectionOwner* BindingsOwner = Private::GetBindingsOwner(InPropertyHandle);
+	if (!BindingsOwner)
+	{
+		return false;
+	}
+
+	// Bound property's children composition gets overridden.
+	FPropertyBindingPath TargetPath;
+	UE::PropertyBinding::MakeStructPropertyPathFromPropertyHandle(InPropertyHandle, TargetPath);
+	if (TargetPath.IsPathEmpty())
+	{
+		return false;
+	}
+
+	return BindingsOwner->GetBindingCollection().HasBinding(TargetPath);
+}
+
+void FBindingsChildrenCustomization::CustomizeChildren(IDetailChildrenBuilder& InChildrenBuilder, TSharedPtr<IPropertyHandle> InPropertyHandle)
+{
+	ISceneStateBindingCollectionOwner* BindingsOwner = Private::GetBindingsOwner(InPropertyHandle);
+	if (!BindingsOwner)
+	{
+		return;
+	}
+
+	FPropertyBindingPath TargetPath;
+	UE::PropertyBinding::MakeStructPropertyPathFromPropertyHandle(InPropertyHandle, TargetPath);
+	if (!FBindingFunctionStructProvider::HasValidBindingFunction(BindingsOwner, TargetPath))
+	{
+		return;
+	}
+
+	// BindingFunction takes control over property's children composition.
+	const TSharedPtr<FBindingFunctionStructProvider> StructProvider = MakeShared<FBindingFunctionStructProvider>(BindingsOwner, TargetPath);
+
+	// Create unique name to persists expansion state.
+	const FName UniqueName(LexToString(TargetPath.GetStructID()) + TargetPath.ToString());
+	InChildrenBuilder.AddChildStructure(InPropertyHandle.ToSharedRef(), StructProvider, UniqueName);
+}
+
+} // UE::SceneState::Editor
diff --git a/Engine/Plugins/VirtualProduction/SceneState/Source/SceneStateBlueprintEditor/Private/DetailsView/SceneStateBindingsChildrenCustomization.h b/Engine/Plugins/VirtualProduction/SceneState/Source/SceneStateBlueprintEditor/Private/DetailsView/SceneStateBindingsChildrenCustomization.h
new file mode 100644
--- /dev/null
+++ b/Engine/Plugins/VirtualProduction/SceneState/Source/SceneStateBlueprintEditor/Private/DetailsView/SceneStateBindingsChildrenCustomization.h
@@ -0,0 +1,48 @@
+﻿// Copyright Epic Games, Inc. All Rights Reserved.
+
+#pragma once
+
+#include "IDetailPropertyChildrenCustomizationHandler.h"
+#include "IStructureDataProvider.h"
+#include "PropertyBindingPath.h"
+#include "UObject/WeakInterfacePtr.h"
+
+class ISceneStateBindingCollectionOwner;
+
+namespace UE::SceneState::Editor
+{
+
+/* Provides the binding function data view for a given target path */
+class FBindingFunctionStructProvider : public IStructureDataProvider
+{
+public:
+	explicit FBindingFunctionStructProvider(ISceneStateBindingCollectionOwner* InBindingsOwnerChecked, FPropertyBindingPath InTargetPath);
+
+	//~ Begin IStructureDataProvider
+	virtual bool IsValid() const override;;
+	virtual const UStruct* GetBaseStructure() const override;
+	virtual void GetInstances(TArray<TSharedPtr<FStructOnScope>>& OutInstances, const UStruct* InExpectedBaseStructure) const override;
+	//~ End IStructureDataProvider
+
+	/** Returns whether there is a valid binding function at the given target path */
+	static bool HasValidBindingFunction(ISceneStateBindingCollectionOwner* InBindingsOwner, const FPropertyBindingPath& InTargetPath);
+
+private:
+	/** Retrieves the binding function at the given target path if any*/
+	static FStructView GetBindingFunctionView(ISceneStateBindingCollectionOwner* InBindingsOwner, const FPropertyBindingPath& InTargetPath);
+
+	TWeakInterfacePtr<ISceneStateBindingCollectionOwner> BindingsOwnerWeak;
+	FPropertyBindingPath TargetPath;
+};
+
+/* Overrides bound property's children composition. */
+class FBindingsChildrenCustomization : public IDetailPropertyChildrenCustomizationHandler
+{
+public:
+	//~ Begin IDetailPropertyChildrenCustomizationHandler
+	virtual bool ShouldCustomizeChildren(TSharedRef<IPropertyHandle> InPropertyHandle) override;
+	virtual void CustomizeChildren(IDetailChildrenBuilder& InChildrenBuilder, TSharedPtr<IPropertyHandle> InPropertyHandle) override;
+	//~ End IDetailPropertyChildrenCustomizationHandler
+};
+
+} // UE::SceneState::Editor
diff --git a/Engine/Plugins/VirtualProduction/SceneState/Source/SceneStateBlueprintEditor/Private/DetailsView/SceneStateCachedBindingData.cpp b/Engine/Plugins/VirtualProduction/SceneState/Source/SceneStateBlueprintEditor/Private/DetailsView/SceneStateCachedBindingData.cpp
--- a/Engine/Plugins/VirtualProduction/SceneState/Source/SceneStateBlueprintEditor/Private/DetailsView/SceneStateCachedBindingData.cpp
+++ b/Engine/Plugins/VirtualProduction/SceneState/Source/SceneStateBlueprintEditor/Private/DetailsView/SceneStateCachedBindingData.cpp
@@ -7,7 +7,10 @@
 #include "Misc/EnumerateRange.h"
 #include "PropertyBindingBindableStructDescriptor.h"
 #include "PropertyBindingDataView.h"
+#include "SceneStateBindingCollection.h"
+#include "SceneStateBindingCollectionOwner.h"
 #include "SceneStateBindingDesc.h"
+#include "SceneStateBindingFunction.h"
 #include "SceneStateBlueprintPropertyReference.h"
 #include "SceneStatePropertyReference.h"
 #include "SceneStatePropertyReferenceMetadata.h"
@@ -19,6 +22,44 @@
 namespace UE::SceneState::Editor
 {
 
+bool IsOutputProperty(const TSharedPtr<IPropertyHandle>& InPropertyHandle)
+{
+	return InPropertyHandle.IsValid() && InPropertyHandle->HasMetaData(TEXT("Output"));
+}
+
+bool IsOutputProperty(const FProperty* InProperty)
+{
+	return InProperty && InProperty->HasMetaData(TEXT("Output"));
+}
+
+const FProperty* FindSingleOutputProperty(const UStruct* InStruct)
+{
+	const FProperty* OutputProperty = nullptr;
+
+	for (const FProperty* Property : TFieldRange<FProperty>(InStruct, EFieldIterationFlags::IncludeSuper))
+	{
+		if (IsOutputProperty(Property))
+		{
+			if (OutputProperty)
+			{
+				return nullptr;
+			}
+			OutputProperty = Property;
+		}
+	}
+
+	return OutputProperty;
+}
+
+const FProperty* FSceneStateCachedBindingData::GetProperty() const
+{
+	if (const IPropertyHandle* PropertyHandlePtr = GetPropertyHandle())
+	{
+		return PropertyHandlePtr->GetProperty();
+	}
+	return nullptr;
+}
+
 void FSceneStateCachedBindingData::UpdatePropertyReferenceTooltip(const FProperty& InProperty, FTextBuilder& InOutTextBuilder) const
 {
 	if (InProperty.HasMetaData(Metadata::IsRefToArray))
@@ -61,7 +102,7 @@ bool FSceneStateCachedBindingData::IsPropertyReference(const FProperty& InProper
 
 void FSceneStateCachedBindingData::UpdateSourcePropertyPath(TConstStructView<FPropertyBindingBindableStructDescriptor> InDescriptor, const FPropertyBindingPath& InSourcePath, FString& OutString)
 {
-	const FPropertyBindingBindableStructDescriptor* Descriptor = InDescriptor.GetPtr();
+	const FSceneStateBindingDesc* Descriptor = InDescriptor.GetPtr<FSceneStateBindingDesc>();
 	if (!Descriptor || !Descriptor->Struct)
 	{
 		return;
@@ -73,10 +114,18 @@ void FSceneStateCachedBindingData::UpdateSourcePropertyPath(TConstStructView<FPr
 		return;
 	}
 
+	int32 FirstSegment = 0;
+
+	// Making first segment of the path invisible for the user if it's property function's single output property.
+	if (Descriptor->DataHandle.IsDataType(ESceneStateDataType::Function) && FindSingleOutputProperty(Descriptor->Struct))
+	{
+		FirstSegment = 1;
+	}
+
 	FStringBuilderBase Result;
-	for (TEnumerateRef<const FPropertyBindingPathSegment> Segment : EnumerateRange(InSourcePath.GetSegments()))
+	for (TEnumerateRef<const FPropertyBindingPathSegment> Segment : EnumerateRange(InSourcePath.GetSegments().RightChop(FirstSegment)))
 	{
-		const int32 SegmentIndex = Segment.GetIndex();
+		const int32 SegmentIndex = Segment.GetIndex() + FirstSegment;
 
 		const FPropertyBindingPathIndirection* Indirection = Indirections.FindByPredicate(
 			[SegmentIndex](const FPropertyBindingPathIndirection& InIndirection)
@@ -89,7 +138,7 @@ void FSceneStateCachedBindingData::UpdateSourcePropertyPath(TConstStructView<FPr
 			return;
 		}
 
-		if (SegmentIndex > 0)
+		if (SegmentIndex > FirstSegment)
 		{
 			Result += TEXT(".");
 		}
@@ -120,6 +169,149 @@ void FSceneStateCachedBindingData::UpdateSourcePropertyPath(TConstStructView<FPr
 	OutString = Result.ToString();
 }
 
+bool FSceneStateCachedBindingData::CanBindToContextStructInternal(const UStruct* InStruct, const int32 InStructIndex)
+{
+	const FSceneStateBindingDesc& BindingDesc = GetBindableStructDescriptor(InStructIndex).Get<FSceneStateBindingDesc>();
+
+	// Binding directly into function's struct is allowed if it contains a compatible single output property.
+	if (BindingDesc.DataHandle.IsDataType(ESceneStateDataType::Function))
+	{
+		const ISceneStateBindingCollectionOwner* BindingOwner = Cast<ISceneStateBindingCollectionOwner>(GetOwner());
+		FPropertyBindingDataView DataView;
+
+		// If DataView exists, struct is an instance of already bound function.
+		if (!BindingOwner || BindingOwner->GetBindingDataViewByID(BindingDesc.ID, DataView))
+		{
+			return false;
+		}
+
+		if (const FStructProperty* StructProperty = CastField<FStructProperty>(GetProperty()))
+		{
+			// Disallow PropertyRef binding to property functions
+			if (StructProperty->Struct && StructProperty->Struct->IsChildOf<FSceneStatePropertyReference>())
+			{
+				return false;
+			}
+		}
+
+		if (const FProperty* OutputProperty = FindSingleOutputProperty(BindingDesc.Struct))
+		{
+			return CanBindToProperty(OutputProperty, { FBindingChainElement(nullptr, InStructIndex), FBindingChainElement(const_cast<FProperty*>(OutputProperty)) });
+		}
+	}
+
+	return Super::CanBindToContextStructInternal(InStruct, InStructIndex);
+}
+
+bool FSceneStateCachedBindingData::CanAcceptPropertyOrChildrenInternal(const FProperty& InSourceProperty, TConstArrayView<FBindingChainElement, int> InBindingChain)
+{
+	if (InBindingChain.IsEmpty())
+	{
+		return Super::CanAcceptPropertyOrChildrenInternal(InSourceProperty, InBindingChain);
+	}
+
+	const int32 StructIndex = InBindingChain[0].ArrayIndex;
+	const FSceneStateBindingDesc& BindingDesc = GetBindableStructDescriptor(StructIndex).Get<FSceneStateBindingDesc>();
+
+	if (BindingDesc.DataHandle.IsDataType(ESceneStateDataType::Function))
+	{
+		const ISceneStateBindingCollectionOwner* BindingOwner = Cast<ISceneStateBindingCollectionOwner>(GetOwner());
+		FPropertyBindingDataView DataView;
+
+		// If DataView exists, struct is an instance of already bound function.
+		if (!BindingOwner || BindingOwner->GetBindingDataViewByID(BindingDesc.ID, DataView))
+		{
+			return false;
+		}
+
+		// To avoid duplicates, PropertyFunction struct's children are not allowed to be bound if it contains a compatible single output property.
+		if (const FProperty* OutputProperty = FindSingleOutputProperty(BindingDesc.Struct))
+		{
+			if (CanBindToProperty(OutputProperty, { FBindingChainElement(nullptr, StructIndex), FBindingChainElement(const_cast<FProperty*>(OutputProperty)) }))
+			{
+				return false;
+			}
+		}
+
+		// Binding to non-output PropertyFunctions properties is not allowed.
+		if (InBindingChain.Num() == 1 && !IsOutputProperty(&InSourceProperty))
+		{
+			return false;
+		}
+	}
+
+	return Super::CanAcceptPropertyOrChildrenInternal(InSourceProperty, InBindingChain);
+}
+
+void FSceneStateCachedBindingData::GetSourceDataViewForNewBinding(TNotNull<IPropertyBindingBindingCollectionOwner*> InBindingsOwner, TConstStructView<FPropertyBindingBindableStructDescriptor> InDescriptor, FPropertyBindingDataView& OutSourceDataView)
+{
+	const FSceneStateBindingDesc& BindingDesc = InDescriptor.Get<FSceneStateBindingDesc>();
+
+	if (BindingDesc.DataHandle.IsDataType(ESceneStateDataType::Function))
+	{
+		OutSourceDataView = FPropertyBindingDataView(BindingDesc.Struct, nullptr);
+	}
+	else
+	{
+		Super::GetSourceDataViewForNewBinding(InBindingsOwner, InDescriptor, OutSourceDataView);
+	}
+}
+
+bool FSceneStateCachedBindingData::AddBindingInternal(TConstStructView<FPropertyBindingBindableStructDescriptor> InDescriptor, FPropertyBindingPath& InOutSourcePath, const FPropertyBindingPath& InTargetPath)
+{
+	const FSceneStateBindingDesc& BindingDesc = InDescriptor.Get<FSceneStateBindingDesc>();
+
+	// Only functions have custom add binding logic
+	if (!BindingDesc.DataHandle.IsDataType(ESceneStateDataType::Function))
+	{
+		return false;
+	}
+
+	ISceneStateBindingCollectionOwner* BindingsOwner = Cast<ISceneStateBindingCollectionOwner>(GetOwner());
+	if (!ensure(BindingsOwner))
+	{
+		return false;
+	}
+
+	UE::SceneState::FBindingFunctionInfo FunctionInfo;
+
+	BindingsOwner->ForEachBindableFunction(
+		[ID = BindingDesc.ID, &FunctionInfo](const FSceneStateBindingDesc& InBindingDesc, const UE::SceneState::FBindingFunctionInfo& InFunctionInfo)->bool
+		{
+			if (InBindingDesc.ID == ID)
+			{
+				FunctionInfo = InFunctionInfo;
+				return false; // break
+			}
+			return true; // continue
+		});
+
+	if (!ensure(FunctionInfo.FunctionType))
+	{
+		return false;
+	}
+
+	FSceneStateBindingCollection& BindingCollection = BindingsOwner->GetBindingCollection();
+
+	// If there are no segments, bindings leads directly into source struct's single output property. Its path has to be recovered.
+	if (InOutSourcePath.NumSegments() == 0)
+	{
+		const FProperty* OutputProperty = FindSingleOutputProperty(BindingDesc.Struct);
+		if (!ensure(OutputProperty))
+		{
+			return false;
+		}
+
+		FPropertyBindingPathSegment OutputPropertySegment = FPropertyBindingPathSegment(OutputProperty->GetFName());
+		InOutSourcePath = BindingCollection.AddBindingFunction(FunctionInfo, MakeArrayView(&OutputPropertySegment, 1), InTargetPath);
+	}
+	else
+	{
+		InOutSourcePath = BindingCollection.AddBindingFunction(FunctionInfo, InOutSourcePath.GetSegments(), InTargetPath);
+	}
+	return true;
+}
+
 void FSceneStateCachedBindingData::AddPropertyInfoOverride(const FProperty& InProperty, TArray<TSharedPtr<const UE::PropertyBinding::FPropertyInfoOverride>>& OutPropertyInfoOverrides) const
 {
 	const FStructProperty* StructProperty = CastField<const FStructProperty>(&InProperty);
diff --git a/Engine/Plugins/VirtualProduction/SceneState/Source/SceneStateBlueprintEditor/Private/DetailsView/SceneStateCachedBindingData.h b/Engine/Plugins/VirtualProduction/SceneState/Source/SceneStateBlueprintEditor/Private/DetailsView/SceneStateCachedBindingData.h
--- a/Engine/Plugins/VirtualProduction/SceneState/Source/SceneStateBlueprintEditor/Private/DetailsView/SceneStateCachedBindingData.h
+++ b/Engine/Plugins/VirtualProduction/SceneState/Source/SceneStateBlueprintEditor/Private/DetailsView/SceneStateCachedBindingData.h
@@ -6,17 +6,36 @@
 
 namespace UE::SceneState::Editor
 {
+	/** Determines whether the given property handle is deemed as an 'output' for a function */
+	bool IsOutputProperty(const TSharedPtr<IPropertyHandle>& InPropertyHandle);
 
+	/** Determines whether the given property is deemed as an 'output' for a function */
+	bool IsOutputProperty(const FProperty* InProperty);
+
+	/** Returns the single property deemed as 'output' for a function. Returns null if none or multiple found. */
+	const FProperty* FindSingleOutputProperty(const UStruct* InStruct);
+
+} // UE::SceneState::Editor
+
+namespace UE::SceneState::Editor
+{
 struct FSceneStateCachedBindingData : UE::PropertyBinding::FCachedBindingData, TSharedFromThis<FSceneStateCachedBindingData>
 {
 	using Super = UE::PropertyBinding::FCachedBindingData;
 	using Super::Super;
 
+	/** Gets the property from the property handle if valid */
+	const FProperty* GetProperty() const;
+
 	//~ Begin UE::PropertyBinding::FCachedBindingData
 	virtual void UpdatePropertyReferenceTooltip(const FProperty& InProperty, FTextBuilder& InOutTextBuilder) const override;
 	virtual bool GetPinTypeAndIconForProperty(const FProperty& InProperty, FPropertyBindingDataView InTargetDataView, FEdGraphPinType& OutPinType, const FSlateBrush*& OutIconBrush) const override;
 	virtual bool IsPropertyReference(const FProperty& InProperty) override;
 	virtual void UpdateSourcePropertyPath(TConstStructView<FPropertyBindingBindableStructDescriptor> InDescriptor, const FPropertyBindingPath& InSourcePath, FString& OutString);
+	virtual bool CanBindToContextStructInternal(const UStruct* InStruct, const int32 InStructIndex) override;
+	virtual bool CanAcceptPropertyOrChildrenInternal(const FProperty& InSourceProperty, TConstArrayView<FBindingChainElement, int> InBindingChain) override;
+	virtual void GetSourceDataViewForNewBinding(TNotNull<IPropertyBindingBindingCollectionOwner*> InBindingsOwner, TConstStructView<FPropertyBindingBindableStructDescriptor> InDescriptor, FPropertyBindingDataView& OutSourceDataView) override;
+	virtual bool AddBindingInternal(TConstStructView<FPropertyBindingBindableStructDescriptor> InDescriptor, FPropertyBindingPath& InOutSourcePath, const FPropertyBindingPath& InTargetPath) override;
 	virtual void AddPropertyInfoOverride(const FProperty& InProperty, TArray<TSharedPtr<const UE::PropertyBinding::FPropertyInfoOverride>>& OutPropertyInfoOverrides) const override;
 	virtual bool DeterminePropertiesCompatibilityInternal(const FProperty* InSourceProperty,const FProperty* InTargetProperty,const void* InSourcePropertyValue,const void* InTargetPropertyValue,bool& bOutAreCompatible) const override;
 	//~ End UE::PropertyBinding::FCachedBindingData
diff --git a/Engine/Plugins/VirtualProduction/SceneState/Source/SceneStateBlueprintEditor/Private/DetailsView/SceneStatePropertyBindingExtension.cpp b/Engine/Plugins/VirtualProduction/SceneState/Source/SceneStateBlueprintEditor/Private/DetailsView/SceneStatePropertyBindingExtension.cpp
--- a/Engine/Plugins/VirtualProduction/SceneState/Source/SceneStateBlueprintEditor/Private/DetailsView/SceneStatePropertyBindingExtension.cpp
+++ b/Engine/Plugins/VirtualProduction/SceneState/Source/SceneStateBlueprintEditor/Private/DetailsView/SceneStatePropertyBindingExtension.cpp
@@ -1,7 +1,10 @@
 // Copyright Epic Games, Inc. All Rights Reserved.
 
 #include "SceneStatePropertyBindingExtension.h"
+#include "EdGraphSchema_K2.h"
 #include "GameFramework/Actor.h"
+#include "IPropertyAccessEditor.h"
+#include "SceneStateBindingDesc.h"
 #include "SceneStateBindingMetadata.h"
 #include "SceneStateBlueprintEditorUtils.h"
 #include "SceneStateCachedBindingData.h"
@@ -67,6 +70,32 @@ bool FBindingExtension::GetPromotionToParameterOverrideInternal(const FProperty&
 	return false;
 }
 
+void FBindingExtension::UpdateContextStruct(TConstStructView<FPropertyBindingBindableStructDescriptor> InStructDesc, FBindingContextStruct& InOutContextStruct, TMap<FString, FText>& InOutSectionNames) const
+{
+	const FSceneStateBindingDesc& BindingDesc = InStructDesc.Get<FSceneStateBindingDesc>();
+
+	// Function overrides it's struct's icon color.
+	if (BindingDesc.DataHandle.IsDataType(ESceneStateDataType::Function))
+	{
+		if (const FProperty* OutputProperty = FindSingleOutputProperty(BindingDesc.Struct))
+		{
+			const UEdGraphSchema_K2* Schema = GetDefault<UEdGraphSchema_K2>();
+
+			FEdGraphPinType PinType;
+			if (ensure(Schema) && Schema->ConvertPropertyToPinType(OutputProperty, PinType))
+			{
+				InOutContextStruct.Color = Schema->GetPinTypeColor(PinType);
+			}
+		}
+	}
+}
+
+bool FBindingExtension::CanBindToProperty(const FPropertyBindingPath& InTargetPath, const IPropertyHandle& InPropertyHandle) const
+{
+	const FProperty* PropertyToBind = InPropertyHandle.GetProperty();
+	return PropertyToBind && !IsOutputProperty(PropertyToBind);
+}
+
 bool FBindingExtension::CanBindToArrayElements() const
 {
 	return true;
diff --git a/Engine/Plugins/VirtualProduction/SceneState/Source/SceneStateBlueprintEditor/Private/DetailsView/SceneStatePropertyBindingExtension.h b/Engine/Plugins/VirtualProduction/SceneState/Source/SceneStateBlueprintEditor/Private/DetailsView/SceneStatePropertyBindingExtension.h
--- a/Engine/Plugins/VirtualProduction/SceneState/Source/SceneStateBlueprintEditor/Private/DetailsView/SceneStatePropertyBindingExtension.h
+++ b/Engine/Plugins/VirtualProduction/SceneState/Source/SceneStateBlueprintEditor/Private/DetailsView/SceneStatePropertyBindingExtension.h
@@ -18,6 +18,8 @@ public:
 	virtual bool IsPropertyExtendable(const UClass* InObjectClass, const IPropertyHandle& InPropertyHandle) const override;
 	virtual TSharedPtr<UE::PropertyBinding::FCachedBindingData> CreateCachedBindingData(IPropertyBindingBindingCollectionOwner* InBindingsOwner, const FPropertyBindingPath& InTargetPath, const TSharedPtr<IPropertyHandle>& InPropertyHandle, TConstArrayView<TInstancedStruct<FPropertyBindingBindableStructDescriptor>> InAccessibleStructs) const override;
 	virtual bool GetPromotionToParameterOverrideInternal(const FProperty& InProperty, bool& bOutOverride) const override;
+	virtual void UpdateContextStruct(TConstStructView<FPropertyBindingBindableStructDescriptor> InStructDesc, FBindingContextStruct& InOutContextStruct, TMap<FString, FText>& InOutSectionNames) const override;
+	virtual bool CanBindToProperty(const FPropertyBindingPath& InTargetPath, const IPropertyHandle& InPropertyHandle) const override;
 	virtual bool CanBindToArrayElements() const override;
 	//~ End FPropertyBindingExtension
 };
diff --git a/Engine/Plugins/VirtualProduction/SceneState/Source/SceneStateBlueprintEditor/Private/SceneStateBlueprintEditor.cpp b/Engine/Plugins/VirtualProduction/SceneState/Source/SceneStateBlueprintEditor/Private/SceneStateBlueprintEditor.cpp
--- a/Engine/Plugins/VirtualProduction/SceneState/Source/SceneStateBlueprintEditor/Private/SceneStateBlueprintEditor.cpp
+++ b/Engine/Plugins/VirtualProduction/SceneState/Source/SceneStateBlueprintEditor/Private/SceneStateBlueprintEditor.cpp
@@ -3,6 +3,7 @@
 #include "SceneStateBlueprintEditor.h"
 #include "AppModes/SceneStateAppModes.h"
 #include "AppModes/SceneStateBlueprintEditorMode.h"
+#include "DetailsView/SceneStateBindingsChildrenCustomization.h"
 #include "DetailsView/SceneStateMachineTaskNodeCustomization.h"
 #include "DetailsView/SceneStatePropertyBindingExtension.h"
 #include "Framework/Application/SlateApplication.h"
@@ -116,6 +117,7 @@ void FSceneStateBlueprintEditor::CreateDefaultTabContents(const TArray<UBlueprin
 			, FOnGetDetailCustomizationInstance::CreateStatic(&FStateMachineTaskNodeCustomization::MakeInstance));
 
 		DetailsView->SetExtensionHandler(MakeShared<FBindingExtension>());
+		DetailsView->SetChildrenCustomizationHandler(MakeShared<FBindingsChildrenCustomization>());
 	}
 }
 
diff --git a/Engine/Plugins/VirtualProduction/SceneState/Source/SceneStateBlueprintEditor/Private/SceneStateBlueprintEditorModule.cpp b/Engine/Plugins/VirtualProduction/SceneState/Source/SceneStateBlueprintEditor/Private/SceneStateBlueprintEditorModule.cpp
--- a/Engine/Plugins/VirtualProduction/SceneState/Source/SceneStateBlueprintEditor/Private/SceneStateBlueprintEditorModule.cpp
+++ b/Engine/Plugins/VirtualProduction/SceneState/Source/SceneStateBlueprintEditor/Private/SceneStateBlueprintEditorModule.cpp
@@ -4,6 +4,7 @@
 #include "AssetToolsModule.h"
 #include "BlueprintEditorModule.h"
 #include "Compiler/SceneStateBlueprintCompilerContext.h"
+#include "DetailsView/SceneStateBindingFunctionCustomization.h"
 #include "DetailsView/SceneStateBlueprintPropertyReferenceCustomization.h"
 #include "DetailsView/SceneStateBlueprintableTaskInstanceCustomization.h"
 #include "DetailsView/SceneStateMachineGraphCustomization.h"
@@ -16,6 +17,7 @@
 #include "Nodes/SceneStateMachineTaskNode.h"
 #include "Nodes/SceneStateMachineTransitionNode.h"
 #include "PropertyEditorModule.h"
+#include "SceneStateBindingFunction.h"
 #include "SceneStateBlueprint.h"
 #include "SceneStateBlueprintEditorCommands.h"
 #include "SceneStateBlueprintEditorLog.h"
@@ -172,6 +174,9 @@ void FBlueprintEditorModule::RegisterDetailCustomizations()
 	PropertyEditorModule.RegisterCustomPropertyTypeLayout(CustomizedTypes.Add_GetRef(FSceneStateBlueprintPropertyReference::StaticStruct()->GetFName())
 		, FOnGetPropertyTypeCustomizationInstance::CreateStatic(&FBlueprintPropertyReferenceCustomization::MakeInstance));
 
+	PropertyEditorModule.RegisterCustomPropertyTypeLayout(CustomizedTypes.Add_GetRef(FSceneStateBindingFunction::StaticStruct()->GetFName())
+		, FOnGetPropertyTypeCustomizationInstance::CreateStatic(&FBindingFunctionCustomization::MakeInstance));
+
 	PropertyEditorModule.RegisterCustomClassLayout(CustomizedClasses.Add_GetRef(USceneStateMachineGraph::StaticClass()->GetFName())
 		, FOnGetDetailCustomizationInstance::CreateStatic(&FStateMachineGraphCustomization::MakeInstance));
 
