Motion Design Scene State: Removed unused function + added some comment/log.

diff --git a/Engine/Plugins/VirtualProduction/SceneState/Source/SceneState/Private/SceneStateUtils.cpp b/Engine/Plugins/VirtualProduction/SceneState/Source/SceneState/Private/SceneStateUtils.cpp
--- a/Engine/Plugins/VirtualProduction/SceneState/Source/SceneState/Private/SceneStateUtils.cpp
+++ b/Engine/Plugins/VirtualProduction/SceneState/Source/SceneState/Private/SceneStateUtils.cpp
@@ -78,12 +78,14 @@ bool ApplyBatch(const FSceneStateExecutionContext& InContext, const FApplyBatchP
 	const FPropertyBindingIndex16 BindingsBatch(InParams.BindingsBatch);
 	if (!BindingsBatch.IsValid())
 	{
+		// Normal behavior if object is not a target of any binding, and so does not have a binding batch.
 		return false;
 	}
 
 	const FSceneStateBindingCollection* const BindingCollection = InContext.GetBindingCollection();
 	if (!BindingCollection)
 	{
+		UE_LOG(LogSceneState, Error, TEXT("[%s] ApplyBatch failed. Binding Collection is invalid."), *InContext.GetExecutionContextName())
 		return false;
 	}
 
diff --git a/Engine/Plugins/VirtualProduction/SceneState/Source/SceneState/Public/SceneStateUtils.h b/Engine/Plugins/VirtualProduction/SceneState/Source/SceneState/Public/SceneStateUtils.h
--- a/Engine/Plugins/VirtualProduction/SceneState/Source/SceneState/Public/SceneStateUtils.h
+++ b/Engine/Plugins/VirtualProduction/SceneState/Source/SceneState/Public/SceneStateUtils.h
@@ -47,14 +47,6 @@ namespace UE::SceneState
 	/** Determines whether the given begin and end forms a valid range */
 	bool IsValidRange(FPropertyBindingIndex16 InBegin, FPropertyBindingIndex16 InEnd);
 
-	struct FSetupBatchParams
-	{
-		/** Batch of the bindings to apply */
-		uint16 BindingsBatch;
-	};
-	/** Sets up a given binding batch to then apply */
-	bool SetupBatch(const FSceneStateExecutionContext& InContext, const FSetupBatchParams& InParams);
-
 	struct FApplyBatchParams
 	{
 		/** Batch of the bindings to apply */
diff --git a/Engine/Plugins/VirtualProduction/SceneState/Source/SceneStateBinding/Private/SceneStateBindingCollection.cpp b/Engine/Plugins/VirtualProduction/SceneState/Source/SceneStateBinding/Private/SceneStateBindingCollection.cpp
--- a/Engine/Plugins/VirtualProduction/SceneState/Source/SceneStateBinding/Private/SceneStateBindingCollection.cpp
+++ b/Engine/Plugins/VirtualProduction/SceneState/Source/SceneStateBinding/Private/SceneStateBindingCollection.cpp
@@ -72,11 +72,8 @@ void FSceneStateBindingCollection::CopyBindingsInternal(const FGuid InFromStruct
 	TArray<TTuple<FGuid, FGuid>, TInlineAllocator<8>> StructIds;
 	StructIds.Emplace(InFromStructId, InToStructId);
 
-	// Worst case: all bindings have unique struct id and are added to the array
-	const int32 IndexCeiling = Bindings.Num();
-
 	// Find the StructId and copy the binding. If we find functions, then copy the function and copy the functions binding recursively.
-	for (int32 Index = 0; Index < StructIds.Num() && ensure(Index < IndexCeiling); ++Index)
+	for (int32 Index = 0; Index < StructIds.Num(); ++Index)
 	{
 		const TTuple<FGuid, FGuid>& StructId = StructIds[Index];
 
diff --git a/Engine/Plugins/VirtualProduction/SceneState/Source/SceneStateBlueprint/Private/SceneStateBlueprintBindingUtils.cpp b/Engine/Plugins/VirtualProduction/SceneState/Source/SceneStateBlueprint/Private/SceneStateBlueprintBindingUtils.cpp
--- a/Engine/Plugins/VirtualProduction/SceneState/Source/SceneStateBlueprint/Private/SceneStateBlueprintBindingUtils.cpp
+++ b/Engine/Plugins/VirtualProduction/SceneState/Source/SceneStateBlueprint/Private/SceneStateBlueprintBindingUtils.cpp
@@ -324,7 +324,7 @@ bool FindBindingDescById(const USceneStateBlueprint& InBlueprint, const FGuid& I
 		return true;
 	}
 
-	return Result == EResult::FoundValid;
+	return false;
 }
 
 bool FindDataViewById(const USceneStateBlueprint& InBlueprint, const FGuid& InStructId, FPropertyBindingDataView& OutDataView)
