Motion Design Scene State: fixed missing reset calls in execution context + made pure virtual type traits editor-only for functions as the only virtuals are wrapped in WITH_EDITOR

diff --git a/Engine/Plugins/VirtualProduction/SceneState/Source/SceneState/Private/SceneStateExecutionContext.cpp b/Engine/Plugins/VirtualProduction/SceneState/Source/SceneState/Private/SceneStateExecutionContext.cpp
--- a/Engine/Plugins/VirtualProduction/SceneState/Source/SceneState/Private/SceneStateExecutionContext.cpp
+++ b/Engine/Plugins/VirtualProduction/SceneState/Source/SceneState/Private/SceneStateExecutionContext.cpp
@@ -92,6 +92,11 @@ void FSceneStateExecutionContext::Reset()
 	StateInstances.Reset();
 	StateMachineInstances.Reset();
 	TransitionInstances.Reset();
+	ConduitInstances.Reset();
+	TaskInstanceContainers.Reset();
+	StateMachineExecutionStack.Reset();
+	FunctionInstanceBatches.Reset();
+	FunctionInstanceViews.Reset();
 
 	if (TSharedPtr<UE::SceneState::FExecutionContextRegistry> ContextRegistry = ContextRegistryWeak.Pin())
 	{
diff --git a/Engine/Plugins/VirtualProduction/SceneState/Source/SceneState/Public/Functions/SceneStateFunction.h b/Engine/Plugins/VirtualProduction/SceneState/Source/SceneState/Public/Functions/SceneStateFunction.h
--- a/Engine/Plugins/VirtualProduction/SceneState/Source/SceneState/Public/Functions/SceneStateFunction.h
+++ b/Engine/Plugins/VirtualProduction/SceneState/Source/SceneState/Public/Functions/SceneStateFunction.h
@@ -92,7 +92,10 @@ struct TStructOpsTypeTraits<FSceneStateFunction> : TStructOpsTypeTraitsBase2<FSc
 {
 	enum
 	{
+	#if WITH_EDITOR
+		// The pure virtual functions in FSceneStateFunction are editor-only
 		WithPureVirtual = true,
+	#endif
 	};
 };
 
