Motion Design: fixed errors outside of editor builds

diff --git a/Engine/Plugins/VirtualProduction/SceneState/Source/SceneState/Public/Functions/SceneStateFunction.h b/Engine/Plugins/VirtualProduction/SceneState/Source/SceneState/Public/Functions/SceneStateFunction.h
--- a/Engine/Plugins/VirtualProduction/SceneState/Source/SceneState/Public/Functions/SceneStateFunction.h
+++ b/Engine/Plugins/VirtualProduction/SceneState/Source/SceneState/Public/Functions/SceneStateFunction.h
@@ -87,16 +87,16 @@ private:
 	friend UE::SceneState::Editor::FBindingCompiler;
 };
 
+#if WITH_EDITOR
+/** The pure virtual functions in FSceneStateFunction are editor-only. */
 template<>
 struct TStructOpsTypeTraits<FSceneStateFunction> : TStructOpsTypeTraitsBase2<FSceneStateFunction>
 {
 	enum
 	{
-	#if WITH_EDITOR
-		// The pure virtual functions in FSceneStateFunction are editor-only
 		WithPureVirtual = true,
-	#endif
 	};
 };
+#endif
 
 #undef UE_API
