Motion Design Scene State: Added core functions to the tasks module.

diff --git a/Engine/Plugins/VirtualProduction/SceneState/Source/SceneStateTasks/Private/Functions/SceneStateFloatFunctions.cpp b/Engine/Plugins/VirtualProduction/SceneState/Source/SceneStateTasks/Private/Functions/SceneStateFloatFunctions.cpp
new file mode 100644
--- /dev/null
+++ b/Engine/Plugins/VirtualProduction/SceneState/Source/SceneStateTasks/Private/Functions/SceneStateFloatFunctions.cpp
@@ -0,0 +1,44 @@
+// Copyright Epic Games, Inc. All Rights Reserved.
+
+#include "Functions/SceneStateFloatFunctions.h"
+#include "SceneStateExecutionContext.h"
+#include "SceneStateTasksLog.h"
+
+#if WITH_EDITOR
+const UScriptStruct* FSceneStateDoubleFunction::OnGetFunctionDataType() const
+{
+	return FInstanceDataType::StaticStruct();
+}
+#endif
+
+void FSceneStateAddDoubleFunction::OnExecute(const FSceneStateExecutionContext& InContext, FStructView InFunctionInstance) const
+{
+	FInstanceDataType& InstanceData = InFunctionInstance.Get<FInstanceDataType>();
+	InstanceData.Output = InstanceData.Left + InstanceData.Right;
+}
+
+void FSceneStateSubtractDoubleFunction::OnExecute(const FSceneStateExecutionContext& InContext, FStructView InFunctionInstance) const
+{
+	FInstanceDataType& InstanceData = InFunctionInstance.Get<FInstanceDataType>();
+	InstanceData.Output = InstanceData.Left - InstanceData.Right;
+}
+
+void FSceneStateMultiplyDoubleFunction::OnExecute(const FSceneStateExecutionContext& InContext, FStructView InFunctionInstance) const
+{
+	FInstanceDataType& InstanceData = InFunctionInstance.Get<FInstanceDataType>();
+	InstanceData.Output = InstanceData.Left * InstanceData.Right;
+}
+
+void FSceneStateDivideDoubleFunction::OnExecute(const FSceneStateExecutionContext& InContext, FStructView InFunctionInstance) const
+{
+	FInstanceDataType& InstanceData = InFunctionInstance.Get<FInstanceDataType>();
+	if (!FMath::IsNearlyZero(InstanceData.Right))
+	{
+		InstanceData.Output = InstanceData.Left / InstanceData.Right;
+	}
+	else
+	{
+		UE_LOG(LogSceneStateTasks, Warning, TEXT("[%s] Attempting to divide %f by zero"), *InContext.GetExecutionContextName(), InstanceData.Left);
+		InstanceData.Output = 0.0;
+	}
+}
diff --git a/Engine/Plugins/VirtualProduction/SceneState/Source/SceneStateTasks/Private/Functions/SceneStateFromStringFunctions.cpp b/Engine/Plugins/VirtualProduction/SceneState/Source/SceneStateTasks/Private/Functions/SceneStateFromStringFunctions.cpp
new file mode 100644
--- /dev/null
+++ b/Engine/Plugins/VirtualProduction/SceneState/Source/SceneStateTasks/Private/Functions/SceneStateFromStringFunctions.cpp
@@ -0,0 +1,15 @@
+// Copyright Epic Games, Inc. All Rights Reserved.
+
+#include "Functions/SceneStateFromStringFunctions.h"
+
+void FSceneStateTextFromStringFunction::OnExecute(const FSceneStateExecutionContext& InContext, FStructView InFunctionInstance) const
+{
+	FInstanceDataType& InstanceData = InFunctionInstance.Get<FInstanceDataType>();
+	InstanceData.Output = FText::FromString(InstanceData.String);
+}
+
+void FSceneStateNameFromStringFunction::OnExecute(const FSceneStateExecutionContext& InContext, FStructView InFunctionInstance) const
+{
+	FInstanceDataType& InstanceData = InFunctionInstance.Get<FInstanceDataType>();
+	InstanceData.Output = FName(InstanceData.String);
+}
diff --git a/Engine/Plugins/VirtualProduction/SceneState/Source/SceneStateTasks/Private/Functions/SceneStateIntegerFunctions.cpp b/Engine/Plugins/VirtualProduction/SceneState/Source/SceneStateTasks/Private/Functions/SceneStateIntegerFunctions.cpp
new file mode 100644
--- /dev/null
+++ b/Engine/Plugins/VirtualProduction/SceneState/Source/SceneStateTasks/Private/Functions/SceneStateIntegerFunctions.cpp
@@ -0,0 +1,44 @@
+// Copyright Epic Games, Inc. All Rights Reserved.
+
+#include "Functions/SceneStateIntegerFunctions.h"
+#include "SceneStateExecutionContext.h"
+#include "SceneStateTasksLog.h"
+
+#if WITH_EDITOR
+const UScriptStruct* FSceneStateIntegerFunction::OnGetFunctionDataType() const
+{
+	return FInstanceDataType::StaticStruct();
+}
+#endif
+
+void FSceneStateAddIntegerFunction::OnExecute(const FSceneStateExecutionContext& InContext, FStructView InFunctionInstance) const
+{
+	FInstanceDataType& InstanceData = InFunctionInstance.Get<FInstanceDataType>();
+	InstanceData.Output = InstanceData.Left + InstanceData.Right;
+}
+
+void FSceneStateSubtractIntegerFunction::OnExecute(const FSceneStateExecutionContext& InContext, FStructView InFunctionInstance) const
+{
+	FInstanceDataType& InstanceData = InFunctionInstance.Get<FInstanceDataType>();
+	InstanceData.Output = InstanceData.Left - InstanceData.Right;
+}
+
+void FSceneStateMultiplyIntegerFunction::OnExecute(const FSceneStateExecutionContext& InContext, FStructView InFunctionInstance) const
+{
+	FInstanceDataType& InstanceData = InFunctionInstance.Get<FInstanceDataType>();
+	InstanceData.Output = InstanceData.Left * InstanceData.Right;
+}
+
+void FSceneStateDivideIntegerFunction::OnExecute(const FSceneStateExecutionContext& InContext, FStructView InFunctionInstance) const
+{
+	FInstanceDataType& InstanceData = InFunctionInstance.Get<FInstanceDataType>();
+	if (InstanceData.Right != 0)
+	{
+		InstanceData.Output = InstanceData.Left / InstanceData.Right;
+	}
+	else
+	{
+		UE_LOG(LogSceneStateTasks, Warning, TEXT("[%s] Attempting to divide %d by zero"), *InContext.GetExecutionContextName(), InstanceData.Left);
+		InstanceData.Output = 0;
+	}
+}
diff --git a/Engine/Plugins/VirtualProduction/SceneState/Source/SceneStateTasks/Private/Functions/SceneStateToStringFunctions.cpp b/Engine/Plugins/VirtualProduction/SceneState/Source/SceneStateTasks/Private/Functions/SceneStateToStringFunctions.cpp
new file mode 100644
--- /dev/null
+++ b/Engine/Plugins/VirtualProduction/SceneState/Source/SceneStateTasks/Private/Functions/SceneStateToStringFunctions.cpp
@@ -0,0 +1,27 @@
+// Copyright Epic Games, Inc. All Rights Reserved.
+
+#include "Functions/SceneStateToStringFunctions.h"
+
+void FSceneStateTextToStringFunction::OnExecute(const FSceneStateExecutionContext& InContext, FStructView InFunctionInstance) const
+{
+	FInstanceDataType& InstanceData = InFunctionInstance.Get<FInstanceDataType>();
+	InstanceData.Output = InstanceData.Value.ToString();
+}
+
+void FSceneStateNameToStringFunction::OnExecute(const FSceneStateExecutionContext& InContext, FStructView InFunctionInstance) const
+{
+	FInstanceDataType& InstanceData = InFunctionInstance.Get<FInstanceDataType>();
+	InstanceData.Output = InstanceData.Value.ToString();
+}
+
+void FSceneStateIntegerToStringFunction::OnExecute(const FSceneStateExecutionContext& InContext, FStructView InFunctionInstance) const
+{
+	FInstanceDataType& InstanceData = InFunctionInstance.Get<FInstanceDataType>();
+	InstanceData.Output = FString::FromInt(InstanceData.Value);
+}
+
+void FSceneStateDoubleToStringFunction::OnExecute(const FSceneStateExecutionContext& InContext, FStructView InFunctionInstance) const
+{
+	FInstanceDataType& InstanceData = InFunctionInstance.Get<FInstanceDataType>();
+	InstanceData.Output = FString::SanitizeFloat(InstanceData.Value);
+}
diff --git a/Engine/Plugins/VirtualProduction/SceneState/Source/SceneStateTasks/Public/Functions/SceneStateFloatFunctions.h b/Engine/Plugins/VirtualProduction/SceneState/Source/SceneStateTasks/Public/Functions/SceneStateFloatFunctions.h
new file mode 100644
--- /dev/null
+++ b/Engine/Plugins/VirtualProduction/SceneState/Source/SceneStateTasks/Public/Functions/SceneStateFloatFunctions.h
@@ -0,0 +1,86 @@
+// Copyright Epic Games, Inc. All Rights Reserved.
+
+#pragma once
+
+#include "Functions/SceneStateFunction.h"
+#include "SceneStateFloatFunctions.generated.h"
+
+/** Function instance for all float operation functions involving two operands, left and right, and a float result */
+USTRUCT()
+struct FSceneStateDoubleFunctionInstance
+{
+	GENERATED_BODY()
+
+	UPROPERTY(EditAnywhere, Category="Core")
+	double Left = 0.0;
+
+	UPROPERTY(EditAnywhere, Category="Core")
+	double Right = 0.0;
+
+	UPROPERTY(EditAnywhere, Category="Core", meta=(Output))
+	double Output = 0.0;
+};
+
+/** Base class for float operations where there are two operands, left and right, and a float result */
+USTRUCT(meta=(Hidden))
+struct FSceneStateDoubleFunction : public FSceneStateFunction
+{
+	GENERATED_BODY()
+
+	using FInstanceDataType = FSceneStateDoubleFunctionInstance;
+
+protected:
+	//~ Begin FSceneStateFunction
+#if WITH_EDITOR
+	virtual const UScriptStruct* OnGetFunctionDataType() const override;
+#endif
+	//~ End FSceneStateFunction
+};
+
+/** Performs addition (Left + Right) */
+USTRUCT(DisplayName="Add", Category="Math")
+struct FSceneStateAddDoubleFunction : public FSceneStateDoubleFunction
+{
+	GENERATED_BODY()
+
+protected:
+	//~ Begin FSceneStateFunction
+	virtual void OnExecute(const FSceneStateExecutionContext& InContext, FStructView InFunctionInstance) const override;
+	//~ End FSceneStateFunction
+};
+
+/** Performs subtraction (Left - Right) */
+USTRUCT(DisplayName="Subtract", Category="Math")
+struct FSceneStateSubtractDoubleFunction : public FSceneStateDoubleFunction
+{
+	GENERATED_BODY()
+
+protected:
+	//~ Begin FSceneStateFunction
+	virtual void OnExecute(const FSceneStateExecutionContext& InContext, FStructView InFunctionInstance) const override;
+	//~ End FSceneStateFunction
+};
+
+/** Performs multiplication (Left * Right) */
+USTRUCT(DisplayName="Multiply", Category="Math")
+struct FSceneStateMultiplyDoubleFunction : public FSceneStateDoubleFunction
+{
+	GENERATED_BODY()
+
+protected:
+	//~ Begin FSceneStateFunction
+	virtual void OnExecute(const FSceneStateExecutionContext& InContext, FStructView InFunctionInstance) const override;
+	//~ End FSceneStateFunction
+};
+
+/** Performs division (Left / Right) */
+USTRUCT(DisplayName="Divide", Category="Math")
+struct FSceneStateDivideDoubleFunction : public FSceneStateDoubleFunction
+{
+	GENERATED_BODY()
+
+protected:
+	//~ Begin FSceneStateFunction
+	virtual void OnExecute(const FSceneStateExecutionContext& InContext, FStructView InFunctionInstance) const override;
+	//~ End FSceneStateFunction
+};
diff --git a/Engine/Plugins/VirtualProduction/SceneState/Source/SceneStateTasks/Public/Functions/SceneStateFromStringFunctions.h b/Engine/Plugins/VirtualProduction/SceneState/Source/SceneStateTasks/Public/Functions/SceneStateFromStringFunctions.h
new file mode 100644
--- /dev/null
+++ b/Engine/Plugins/VirtualProduction/SceneState/Source/SceneStateTasks/Public/Functions/SceneStateFromStringFunctions.h
@@ -0,0 +1,78 @@
+// Copyright Epic Games, Inc. All Rights Reserved.
+
+#pragma once
+
+#include "Functions/SceneStateFunction.h"
+#include "SceneStateFromStringFunctions.generated.h"
+
+/** Base function instance for all "FromString" functions */
+USTRUCT()
+struct FSceneStateFromStringFunctionInstance
+{
+	GENERATED_BODY()
+
+	UPROPERTY(EditAnywhere, Category="Core")
+	FString String;
+};
+
+//////////////////////////////////////////////////////////////////////////////////////////////
+// Text from String
+//////////////////////////////////////////////////////////////////////////////////////////////
+USTRUCT()
+struct FSceneStateTextFromStringFunctionInstance : public FSceneStateFromStringFunctionInstance
+{
+	GENERATED_BODY()
+
+	UPROPERTY(EditAnywhere, Category="Core", meta=(Output))
+	FText Output;
+};
+
+USTRUCT(DisplayName="String to Text", Category="Core")
+struct FSceneStateTextFromStringFunction : public FSceneStateFunction
+{
+	GENERATED_BODY()
+
+	using FInstanceDataType = FSceneStateTextFromStringFunctionInstance;
+
+protected:
+	//~ Begin FSceneStateFunction
+#if WITH_EDITOR
+	virtual const UScriptStruct* OnGetFunctionDataType() const override
+	{
+		return FInstanceDataType::StaticStruct();
+	}
+#endif
+	virtual void OnExecute(const FSceneStateExecutionContext& InContext, FStructView InFunctionInstance) const override;
+	//~ End FSceneStateFunction
+};
+
+//////////////////////////////////////////////////////////////////////////////////////////////
+// Name from String
+//////////////////////////////////////////////////////////////////////////////////////////////
+USTRUCT()
+struct FSceneStateNameFromStringFunctionInstance : public FSceneStateFromStringFunctionInstance
+{
+	GENERATED_BODY()
+
+	UPROPERTY(EditAnywhere, Category="Core", meta=(Output))
+	FName Output;
+};
+
+USTRUCT(DisplayName="String to Name", Category="Core")
+struct FSceneStateNameFromStringFunction : public FSceneStateFunction
+{
+	GENERATED_BODY()
+
+	using FInstanceDataType = FSceneStateNameFromStringFunctionInstance;
+
+protected:
+	//~ Begin FSceneStateFunction
+#if WITH_EDITOR
+	virtual const UScriptStruct* OnGetFunctionDataType() const override
+	{
+		return FInstanceDataType::StaticStruct();
+	}
+#endif
+	virtual void OnExecute(const FSceneStateExecutionContext& InContext, FStructView InFunctionInstance) const override;
+	//~ End FSceneStateFunction
+};
diff --git a/Engine/Plugins/VirtualProduction/SceneState/Source/SceneStateTasks/Public/Functions/SceneStateIntegerFunctions.h b/Engine/Plugins/VirtualProduction/SceneState/Source/SceneStateTasks/Public/Functions/SceneStateIntegerFunctions.h
new file mode 100644
--- /dev/null
+++ b/Engine/Plugins/VirtualProduction/SceneState/Source/SceneStateTasks/Public/Functions/SceneStateIntegerFunctions.h
@@ -0,0 +1,86 @@
+// Copyright Epic Games, Inc. All Rights Reserved.
+
+#pragma once
+
+#include "Functions/SceneStateFunction.h"
+#include "SceneStateIntegerFunctions.generated.h"
+
+/** Function instance for all integer operation functions involving two operands, left and right, and an integer result */
+USTRUCT()
+struct FSceneStateIntegerFunctionInstance
+{
+	GENERATED_BODY()
+
+	UPROPERTY(EditAnywhere, Category="Core")
+	int32 Left = 0;
+
+	UPROPERTY(EditAnywhere, Category="Core")
+	int32 Right = 0;
+
+	UPROPERTY(EditAnywhere, Category="Core", meta=(Output))
+	int32 Output = 0;
+};
+
+/** Base class for integer operations where there are two operands, left and right, and an integer result */
+USTRUCT(meta=(Hidden))
+struct FSceneStateIntegerFunction : public FSceneStateFunction
+{
+	GENERATED_BODY()
+
+	using FInstanceDataType = FSceneStateIntegerFunctionInstance;
+
+protected:
+	//~ Begin FSceneStateFunction
+#if WITH_EDITOR
+	virtual const UScriptStruct* OnGetFunctionDataType() const override;
+#endif
+	//~ End FSceneStateFunction
+};
+
+/** Performs addition (Left + Right) */
+USTRUCT(DisplayName="Add", Category="Math")
+struct FSceneStateAddIntegerFunction : public FSceneStateIntegerFunction
+{
+	GENERATED_BODY()
+
+protected:
+	//~ Begin FSceneStateFunction
+	virtual void OnExecute(const FSceneStateExecutionContext& InContext, FStructView InFunctionInstance) const override;
+	//~ End FSceneStateFunction
+};
+
+/** Performs subtraction (Left - Right) */
+USTRUCT(DisplayName="Subtract", Category="Math")
+struct FSceneStateSubtractIntegerFunction : public FSceneStateIntegerFunction
+{
+	GENERATED_BODY()
+
+protected:
+	//~ Begin FSceneStateFunction
+	virtual void OnExecute(const FSceneStateExecutionContext& InContext, FStructView InFunctionInstance) const override;
+	//~ End FSceneStateFunction
+};
+
+/** Performs multiplication (Left * Right) */
+USTRUCT(DisplayName="Multiply", Category="Math")
+struct FSceneStateMultiplyIntegerFunction : public FSceneStateIntegerFunction
+{
+	GENERATED_BODY()
+
+protected:
+	//~ Begin FSceneStateFunction
+	virtual void OnExecute(const FSceneStateExecutionContext& InContext, FStructView InFunctionInstance) const override;
+	//~ End FSceneStateFunction
+};
+
+/** Performs division (Left / Right) */
+USTRUCT(DisplayName="Divide", Category="Math")
+struct FSceneStateDivideIntegerFunction : public FSceneStateIntegerFunction
+{
+	GENERATED_BODY()
+
+protected:
+	//~ Begin FSceneStateFunction
+	virtual void OnExecute(const FSceneStateExecutionContext& InContext, FStructView InFunctionInstance) const override;
+	//~ End FSceneStateFunction
+};
diff --git a/Engine/Plugins/VirtualProduction/SceneState/Source/SceneStateTasks/Public/Functions/SceneStateToStringFunctions.h b/Engine/Plugins/VirtualProduction/SceneState/Source/SceneStateTasks/Public/Functions/SceneStateToStringFunctions.h
new file mode 100644
--- /dev/null
+++ b/Engine/Plugins/VirtualProduction/SceneState/Source/SceneStateTasks/Public/Functions/SceneStateToStringFunctions.h
@@ -0,0 +1,140 @@
+// Copyright Epic Games, Inc. All Rights Reserved.
+
+#pragma once
+
+#include "Functions/SceneStateFunction.h"
+#include "SceneStateToStringFunctions.generated.h"
+
+/** Base function instance for all "ToString" functions */
+USTRUCT()
+struct FSceneStateToStringFunctionInstance
+{
+	GENERATED_BODY()
+
+	UPROPERTY(EditAnywhere, Category="Core", meta=(Output))
+	FString Output;
+};
+
+//////////////////////////////////////////////////////////////////////////////////////////////
+// Text to String
+//////////////////////////////////////////////////////////////////////////////////////////////
+USTRUCT()
+struct FSceneStateTextToStringFunctionInstance : public FSceneStateToStringFunctionInstance
+{
+	GENERATED_BODY()
+
+	UPROPERTY(EditAnywhere, Category="Core")
+	FText Value;
+};
+
+USTRUCT(DisplayName="Text to String", Category="Core")
+struct FSceneStateTextToStringFunction : public FSceneStateFunction
+{
+	GENERATED_BODY()
+
+	using FInstanceDataType = FSceneStateTextToStringFunctionInstance;
+
+protected:
+	//~ Begin FSceneStateFunction
+#if WITH_EDITOR
+	virtual const UScriptStruct* OnGetFunctionDataType() const override
+	{
+		return FInstanceDataType::StaticStruct();
+	}
+#endif
+	virtual void OnExecute(const FSceneStateExecutionContext& InContext, FStructView InFunctionInstance) const override;
+	//~ End FSceneStateFunction
+};
+
+//////////////////////////////////////////////////////////////////////////////////////////////
+// Name to String
+//////////////////////////////////////////////////////////////////////////////////////////////
+USTRUCT()
+struct FSceneStateNameToStringFunctionInstance : public FSceneStateToStringFunctionInstance
+{
+	GENERATED_BODY()
+
+	UPROPERTY(EditAnywhere, Category="Core")
+	FName Value;
+};
+
+USTRUCT(DisplayName="Name to String", Category="Core")
+struct FSceneStateNameToStringFunction : public FSceneStateFunction
+{
+	GENERATED_BODY()
+
+	using FInstanceDataType = FSceneStateNameToStringFunctionInstance;
+
+protected:
+	//~ Begin FSceneStateFunction
+#if WITH_EDITOR
+	virtual const UScriptStruct* OnGetFunctionDataType() const override
+	{
+		return FInstanceDataType::StaticStruct();
+	}
+#endif
+	virtual void OnExecute(const FSceneStateExecutionContext& InContext, FStructView InFunctionInstance) const override;
+	//~ End FSceneStateFunction
+};
+
+//////////////////////////////////////////////////////////////////////////////////////////////
+// Integer to String
+//////////////////////////////////////////////////////////////////////////////////////////////
+USTRUCT()
+struct FSceneStateIntegerToStringFunctionInstance : public FSceneStateToStringFunctionInstance
+{
+	GENERATED_BODY()
+
+	UPROPERTY(EditAnywhere, Category="Core")
+	int32 Value = 0;
+};
+
+USTRUCT(DisplayName="Integer to String", Category="Core")
+struct FSceneStateIntegerToStringFunction : public FSceneStateFunction
+{
+	GENERATED_BODY()
+
+	using FInstanceDataType = FSceneStateIntegerToStringFunctionInstance;
+
+protected:
+	//~ Begin FSceneStateFunction
+#if WITH_EDITOR
+	virtual const UScriptStruct* OnGetFunctionDataType() const override
+	{
+		return FInstanceDataType::StaticStruct();
+	}
+#endif
+	virtual void OnExecute(const FSceneStateExecutionContext& InContext, FStructView InFunctionInstance) const override;
+	//~ End FSceneStateFunction
+};
+
+//////////////////////////////////////////////////////////////////////////////////////////////
+// Float to String
+//////////////////////////////////////////////////////////////////////////////////////////////
+USTRUCT()
+struct FSceneStateDoubleToStringFunctionInstance : public FSceneStateToStringFunctionInstance
+{
+	GENERATED_BODY()
+
+	UPROPERTY(EditAnywhere, Category="Core")
+	double Value = 0.0;
+};
+
+USTRUCT(DisplayName="Float to String", Category="Core")
+struct FSceneStateDoubleToStringFunction : public FSceneStateFunction
+{
+	GENERATED_BODY()
+
+	using FInstanceDataType = FSceneStateDoubleToStringFunctionInstance;
+
+protected:
+	//~ Begin FSceneStateFunction
+#if WITH_EDITOR
+	virtual const UScriptStruct* OnGetFunctionDataType() const override
+	{
+		return FInstanceDataType::StaticStruct();
+	}
+#endif
+	virtual void OnExecute(const FSceneStateExecutionContext& InContext, FStructView InFunctionInstance) const override;
+	//~ End FSceneStateFunction
+};
