Motion Design: updated comments for clarity on double-precision floats

diff --git a/Engine/Plugins/VirtualProduction/SceneState/Source/SceneStateTasks/Public/Functions/SceneStateFloatFunctions.h b/Engine/Plugins/VirtualProduction/SceneState/Source/SceneStateTasks/Public/Functions/SceneStateFloatFunctions.h
--- a/Engine/Plugins/VirtualProduction/SceneState/Source/SceneStateTasks/Public/Functions/SceneStateFloatFunctions.h
+++ b/Engine/Plugins/VirtualProduction/SceneState/Source/SceneStateTasks/Public/Functions/SceneStateFloatFunctions.h
@@ -5,7 +5,7 @@
 #include "Functions/SceneStateFunction.h"
 #include "SceneStateFloatFunctions.generated.h"
 
-/** Function instance for all float operation functions involving two operands, left and right, and a float result */
+/** Function instance for all double-precision float operation functions involving two operands, left and right, and a result */
 USTRUCT()
 struct FSceneStateDoubleFunctionInstance
 {
@@ -21,7 +21,7 @@ struct FSceneStateDoubleFunctionInstance
 	double Output = 0.0;
 };
 
-/** Base class for float operations where there are two operands, left and right, and a float result */
+/** Base class for double-precision float operations where there are two operands, left and right, and a result */
 USTRUCT(meta=(Hidden))
 struct FSceneStateDoubleFunction : public FSceneStateFunction
 {
diff --git a/Engine/Plugins/VirtualProduction/SceneState/Source/SceneStateTasks/Public/Functions/SceneStateToStringFunctions.h b/Engine/Plugins/VirtualProduction/SceneState/Source/SceneStateTasks/Public/Functions/SceneStateToStringFunctions.h
--- a/Engine/Plugins/VirtualProduction/SceneState/Source/SceneStateTasks/Public/Functions/SceneStateToStringFunctions.h
+++ b/Engine/Plugins/VirtualProduction/SceneState/Source/SceneStateTasks/Public/Functions/SceneStateToStringFunctions.h
@@ -109,7 +109,7 @@ protected:
 };
 
 //////////////////////////////////////////////////////////////////////////////////////////////
-// Float to String
+// Double-precision Float to String
 //////////////////////////////////////////////////////////////////////////////////////////////
 USTRUCT()
 struct FSceneStateDoubleToStringFunctionInstance : public FSceneStateToStringFunctionInstance
