Motion Design Scene State: updated property functions look in editor ui

diff --git a/Engine/Plugins/VirtualProduction/SceneState/Source/SceneStateBinding/Public/SceneStateBindingDesc.h b/Engine/Plugins/VirtualProduction/SceneState/Source/SceneStateBinding/Public/SceneStateBindingDesc.h
--- a/Engine/Plugins/VirtualProduction/SceneState/Source/SceneStateBinding/Public/SceneStateBindingDesc.h
+++ b/Engine/Plugins/VirtualProduction/SceneState/Source/SceneStateBinding/Public/SceneStateBindingDesc.h
@@ -23,4 +23,10 @@ struct FSceneStateBindingDesc : public FPropertyBindingBindableStructDescriptor
 
 	UPROPERTY()
 	FSceneStateBindingDataHandle DataHandle;
+
+#if WITH_EDITORONLY_DATA
+	/** Section of the bindable struct. Can be used to group bindable structs within the same category menu. */
+	UPROPERTY()
+	FText Section;
+#endif
 };
diff --git a/Engine/Plugins/VirtualProduction/SceneState/Source/SceneStateBlueprint/Private/SceneStateBlueprint.cpp b/Engine/Plugins/VirtualProduction/SceneState/Source/SceneStateBlueprint/Private/SceneStateBlueprint.cpp
--- a/Engine/Plugins/VirtualProduction/SceneState/Source/SceneStateBlueprint/Private/SceneStateBlueprint.cpp
+++ b/Engine/Plugins/VirtualProduction/SceneState/Source/SceneStateBlueprint/Private/SceneStateBlueprint.cpp
@@ -21,6 +21,8 @@
 #include "SceneStateMachineGraph.h"
 #include "UObject/UObjectIterator.h"
 
+#define LOCTEXT_NAMESPACE "SceneStateBlueprint"
+
 USceneStateBlueprint::USceneStateBlueprint(const FObjectInitializer& InObjectInitializer)
 	: Super(InObjectInitializer)
 {
@@ -59,6 +61,61 @@ FSceneStateBindingDesc USceneStateBlueprint::CreateRootBinding() const
 	return CreateBindingDesc(*this);
 }
 
+void USceneStateBlueprint::GetBindingStructs(const FGetBindableStructsParams& InParams, TArray<TInstancedStruct<FPropertyBindingBindableStructDescriptor>>& OutStructDescs) const
+{
+	using namespace UE::SceneState::Graph;
+
+	TArray<TInstancedStruct<FSceneStateBindingDesc>> BindingDescs;
+
+	ON_SCOPE_EXIT
+	{
+		OutStructDescs.Append(MoveTemp(BindingDescs));
+	};
+
+	auto GatherDefaultBindingDescs =
+		[This=this, &InParams, &BindingDescs]()
+		{
+			GetGlobalBindingDescs(*This, InParams.TargetStructId, BindingDescs);
+			if (InParams.bIncludeFunctions)
+			{
+				GetFunctionBindingDescs(*This, InParams.TargetStructId, BindingDescs);
+			}
+		};
+
+	// Add State Machine Bindable Structs if the Target Struct Id matches the State Machine Id
+	if (USceneStateMachineGraph* SceneStateMachine = FindStateMachineMatchingId(*this, InParams.TargetStructId))
+	{
+		GatherDefaultBindingDescs();
+		GetStateMachineBindingDescs(*this, *SceneStateMachine, BindingDescs);
+		return;
+	}
+
+	// Add Transition Bindable Structs if the Target Struct Id matches the Transition Parameters Id
+	if (USceneStateMachineTransitionNode* TransitionNode = FindTransitionMatchingId(*this, InParams.TargetStructId))
+	{
+		GatherDefaultBindingDescs();
+		GetTransitionBindingDescs(*this, *TransitionNode, BindingDescs);
+		return;
+	}
+
+	// Add the Task Bindable Structs if the target struct is either:
+	//   - The task itself (so the task id must match the target id)
+	//   - A struct within a Task Instance that is not stable across sessions, and so require their own ID.
+	//     One example of this Instanced Property Bags, as the underlying UPropertyBag (UScriptStruct) is not serialized.
+	if (const USceneStateMachineTaskNode* TargetTaskNode = FindTaskNodeContainingId(*this, InParams.TargetStructId))
+	{
+		GatherDefaultBindingDescs();
+		GetTaskBindingDescs(*this, *TargetTaskNode, BindingDescs);
+		return;
+	}
+
+	// Add only default binding descs if the target struct is a binding function
+	if (const FSceneStateBindingFunction* BindingFunction = FindBindingFunctionMatchingId(*this, InParams.TargetStructId))
+	{
+		GatherDefaultBindingDescs();
+	}
+}
+
 FSceneStateBindingCollection& USceneStateBlueprint::GetBindingCollection()
 {
 	return BindingCollection;
@@ -74,7 +131,7 @@ bool USceneStateBlueprint::ForEachBindableFunction(TFunctionRef<bool(const FScen
 	const FName MD_Hidden = TEXT("Hidden");
 	const FName MD_Category = TEXT("Category");
 
-	const FString BaseCategory = TEXT("Functions|");
+	const FText Section = LOCTEXT("BindingFunctionsSectionLabel", "Functions");
 
 	// Todo: cache functions for faster iteration
 	for (const UScriptStruct* Struct : TObjectRange<UScriptStruct>())
@@ -90,7 +147,8 @@ bool USceneStateBlueprint::ForEachBindableFunction(TFunctionRef<bool(const FScen
 		BindingDesc.Struct = InstanceType;
 		BindingDesc.Name = *Struct->GetDisplayNameText().ToString();
 		BindingDesc.ID = FGuid::NewDeterministicGuid(Struct->GetStructPathName().ToString());
-		BindingDesc.Category = BaseCategory + Struct->GetMetaData(MD_Category);
+		BindingDesc.Section = Section;
+		BindingDesc.Category = Struct->GetMetaData(MD_Category);
 		BindingDesc.DataHandle = FSceneStateBindingDataHandle(ESceneStateDataType::Function);
 
 		UE::SceneState::FBindingFunctionInfo FunctionInfo;
@@ -124,48 +182,12 @@ bool USceneStateBlueprint::ForEachBindingFunction(TFunctionRef<bool(const FScene
 
 void USceneStateBlueprint::GetBindableStructs(const FGuid InTargetStructId, TArray<TInstancedStruct<FPropertyBindingBindableStructDescriptor>>& OutStructDescs) const
 {
-	using namespace UE::SceneState::Graph;
-
-	// Add State Machine Bindable Structs if the Target Struct Id matches the State Machine Id
-	if (USceneStateMachineGraph* SceneStateMachine = FindStateMachineMatchingId(*this, InTargetStructId))
-	{
-		TArray<TInstancedStruct<FSceneStateBindingDesc>> BindingDescs;
-		GetGlobalBindingDescs(*this, InTargetStructId, BindingDescs);
-		GetStateMachineBindingDescs(*this, *SceneStateMachine, BindingDescs);
-		OutStructDescs.Append(MoveTemp(BindingDescs));
-		return;
-	}
-
-	// Add Transition Bindable Structs if the Target Struct Id matches the Transition Parameters Id
-	if (USceneStateMachineTransitionNode* TransitionNode = FindTransitionMatchingId(*this, InTargetStructId))
-	{
-		TArray<TInstancedStruct<FSceneStateBindingDesc>> BindingDescs;
-		GetGlobalBindingDescs(*this, InTargetStructId, BindingDescs);
-		GetTransitionBindingDescs(*this, *TransitionNode, BindingDescs);
-		OutStructDescs.Append(MoveTemp(BindingDescs));
-		return;
-	}
-
-	// Add the Task Bindable Structs if the target struct is either:
-	//   - The task itself (so the task id must match the target id)
-	//   - A struct within a Task Instance that is not stable across sessions, and so require their own ID.
-	//     One example of this Instanced Property Bags, as the underlying UPropertyBag (UScriptStruct) is not serialized.
-	if (const USceneStateMachineTaskNode* TargetTaskNode = FindTaskNodeContainingId(*this, InTargetStructId))
-	{
-		TArray<TInstancedStruct<FSceneStateBindingDesc>> BindingDescs;
-		GetGlobalBindingDescs(*this, InTargetStructId, BindingDescs);
-		GetTaskBindingDescs(*this, *TargetTaskNode, BindingDescs);
-		OutStructDescs.Append(MoveTemp(BindingDescs));
-		return;
-	}
-
-	// Add global bindable structs if the target struct is a binding function
-	if (const FSceneStateBindingFunction* BindingFunction = FindBindingFunctionMatchingId(*this, InTargetStructId))
-	{
-		TArray<TInstancedStruct<FSceneStateBindingDesc>> BindingDescs;
-		GetGlobalBindingDescs(*this, InTargetStructId, BindingDescs);
-		OutStructDescs.Append(MoveTemp(BindingDescs));
-	}
+	const FGetBindableStructsParams Params
+		{
+			.TargetStructId = InTargetStructId,
+			.bIncludeFunctions = false,
+		};
+	GetBindingStructs(Params, OutStructDescs);
 }
 
 bool USceneStateBlueprint::GetBindableStructByID(const FGuid InStructId, TInstancedStruct<FPropertyBindingBindableStructDescriptor>& OutStructDesc) const
@@ -379,3 +401,5 @@ void USceneStateBlueprint::UpdateGraphParametersBindings(FPropertyBindingPath& I
 		}
 	}
 }
+
+#undef LOCTEXT_NAMESPACE
diff --git a/Engine/Plugins/VirtualProduction/SceneState/Source/SceneStateBlueprint/Private/SceneStateBlueprintBindingUtils.cpp b/Engine/Plugins/VirtualProduction/SceneState/Source/SceneStateBlueprint/Private/SceneStateBlueprintBindingUtils.cpp
--- a/Engine/Plugins/VirtualProduction/SceneState/Source/SceneStateBlueprint/Private/SceneStateBlueprintBindingUtils.cpp
+++ b/Engine/Plugins/VirtualProduction/SceneState/Source/SceneStateBlueprint/Private/SceneStateBlueprintBindingUtils.cpp
@@ -16,21 +16,29 @@
 #include "StructUtils/UserDefinedStruct.h"
 #include "Tasks/SceneStateBlueprintableTask.h"
 
+#define LOCTEXT_NAMESPACE "SceneStateBlueprintBindingUtils"
+
 namespace UE::SceneState::Graph
 {
 
-void GetStateBindingDescs(const USceneStateMachineStateNode* InStateNode, TArray<TInstancedStruct<FSceneStateBindingDesc>>& OutBindingDescs, FString InBaseCategory)
+namespace Private
+{
+	
+FText GetParametersSectionLabel()
+{
+	return LOCTEXT("ParametersSectionLabel", "Parameters");
+}
+
+} // UE::SceneState::Graph::Private
+
+void GetStateBindingDescs(const USceneStateMachineStateNode* InStateNode, TArray<TInstancedStruct<FSceneStateBindingDesc>>& OutBindingDescs, FString InCategory)
 {
 	if (!InStateNode)
 	{
 		return;
 	}
 
-	const FString CategoryName = InBaseCategory.IsEmpty()
-		? InStateNode->GetNodeName().ToString()
-		: InBaseCategory + TEXT("|") + InStateNode->GetNodeName().ToString();
-
-	const FString EventCategoryName = CategoryName + TEXT(" Events");
+	const FText Section = LOCTEXT("EventsSectionLabel", "Events");
 
 	// Add the State Event Handlers as Bindable Sources
 	for (const FSceneStateEventHandler& EventHandler : InStateNode->GetEventHandlers())
@@ -41,7 +49,8 @@ void GetStateBindingDescs(const USceneStateMachineStateNode* InStateNode, TArray
 			BindingDesc.ID = EventHandler.GetHandlerId();
 			BindingDesc.Name = EventSchema->Name;
 			BindingDesc.Struct = EventSchema->Struct;
-			BindingDesc.Category = EventCategoryName;
+			BindingDesc.Category = InCategory;
+			BindingDesc.Section = Section;
 
 			OutBindingDescs.Emplace(TInstancedStruct<FSceneStateBindingDesc>::Make(MoveTemp(BindingDesc)));
 		}
@@ -57,7 +66,12 @@ void GetStateBindingDescs(const USceneStateMachineStateNode* InStateNode, TArray
 	OutBindingDescs.Emplace(TInstancedStruct<FSceneStateBindingDesc>::Make(CreateBindingDesc(*StateMachineGraph)));
 
 	// Recurse up to add the Event Handlers of the parent state node (if any)
-	GetStateBindingDescs(StateMachineGraph->GetParentStateNode(), OutBindingDescs, CategoryName);
+	if (InCategory.IsEmpty())
+	{
+		InCategory += TEXT("|");
+	}
+	InCategory += InStateNode->GetNodeName().ToString();
+	GetStateBindingDescs(StateMachineGraph->GetParentStateNode(), OutBindingDescs, InCategory);
 }
 
 FSceneStateBindingDesc CreateBindingDesc(const USceneStateBlueprint& InBlueprint)
@@ -66,6 +80,7 @@ FSceneStateBindingDesc CreateBindingDesc(const USceneStateBlueprint& InBlueprint
 	BindingDesc.Name = TEXT("Variables");
 	BindingDesc.ID = InBlueprint.GetRootId();
 	BindingDesc.Struct = InBlueprint.GeneratedClass;
+	BindingDesc.Section = LOCTEXT("BlueprintSectionLabel", "Blueprint");
 	return BindingDesc;
 }
 
@@ -75,6 +90,7 @@ FSceneStateBindingDesc CreateBindingDesc(const USceneStateMachineGraph& InGraph)
 	BindingDesc.Name = TEXT("State Machine Parameters");
 	BindingDesc.ID = InGraph.ParametersId;
 	BindingDesc.Struct = InGraph.Parameters.GetPropertyBagStruct();
+	BindingDesc.Section = Private::GetParametersSectionLabel();
 	return BindingDesc;
 }
 
@@ -84,6 +100,7 @@ FSceneStateBindingDesc CreateBindingDesc(const USceneStateMachineTransitionNode&
 	BindingDesc.Name = TEXT("Transition Parameters");
 	BindingDesc.ID = InTransitionNode.GetParametersId();
 	BindingDesc.Struct = InTransitionNode.GetParameters().GetPropertyBagStruct();
+	BindingDesc.Section = Private::GetParametersSectionLabel();
 	return BindingDesc;
 }
 
@@ -99,7 +116,10 @@ FSceneStateBindingDesc CreateBindingDesc(const FSceneStateBindingFunction& InBin
 void GetGlobalBindingDescs(const USceneStateBlueprint& InBlueprint, const FGuid& InTargetStructId, TArray<TInstancedStruct<FSceneStateBindingDesc>>& OutBindingDescs)
 {
 	OutBindingDescs.Add(TInstancedStruct<FSceneStateBindingDesc>::Make(CreateBindingDesc(InBlueprint)));
+}
 
+void GetFunctionBindingDescs(const USceneStateBlueprint& InBlueprint, const FGuid& InTargetStructId, TArray<TInstancedStruct<FSceneStateBindingDesc>>& OutBindingDescs)
+{
 	for (const FPropertyBindingBinding& Binding : InBlueprint.GetBindingCollection().GetBindings())
 	{
 		if (Binding.GetTargetPath().GetStructID() != InTargetStructId)
@@ -410,3 +430,5 @@ bool FindDataViewById(const USceneStateBlueprint& InBlueprint, const FGuid& InSt
 }
 
 } // namespace UE::SceneState::Graph
+
+#undef LOCTEXT_NAMESPACE
diff --git a/Engine/Plugins/VirtualProduction/SceneState/Source/SceneStateBlueprint/Private/SceneStateBlueprintBindingUtils.h b/Engine/Plugins/VirtualProduction/SceneState/Source/SceneStateBlueprint/Private/SceneStateBlueprintBindingUtils.h
--- a/Engine/Plugins/VirtualProduction/SceneState/Source/SceneStateBlueprint/Private/SceneStateBlueprintBindingUtils.h
+++ b/Engine/Plugins/VirtualProduction/SceneState/Source/SceneStateBlueprint/Private/SceneStateBlueprintBindingUtils.h
@@ -23,9 +23,9 @@ namespace UE::SceneState::Graph
 	 * Gathers the Binding Descs for a State Node
 	 * @param InStateNode the state to find binding descs for
 	 * @param OutBindingDescs the returned binding descs
-	 * @param InBaseCategory the category to append to the binding descs
+	 * @param InCategory the category to append to the binding descs
 	 */
-	void GetStateBindingDescs(const USceneStateMachineStateNode* InStateNode, TArray<TInstancedStruct<FSceneStateBindingDesc>>& OutBindingDescs, FString InBaseCategory);
+	void GetStateBindingDescs(const USceneStateMachineStateNode* InStateNode, TArray<TInstancedStruct<FSceneStateBindingDesc>>& OutBindingDescs, FString InCategory);
 
 	/** Creates the Binding Struct Desc for the Variables in the given Blueprint */
 	FSceneStateBindingDesc CreateBindingDesc(const USceneStateBlueprint& InBlueprint);
@@ -47,6 +47,14 @@ namespace UE::SceneState::Graph
 	 */
 	void GetGlobalBindingDescs(const USceneStateBlueprint& InBlueprint, const FGuid& InTargetStructId, TArray<TInstancedStruct<FSceneStateBindingDesc>>& OutBindingDescs);
 
+	/**
+	 * Gathers all the binding descs for functions bound to the given struct id.
+	 * @param InBlueprint the blueprint
+	 * @param InTargetStructId the struct to gather the binding descs for
+	 * @param OutBindingDescs all the binding descs that are bound to the target struct
+	 */
+	void GetFunctionBindingDescs(const USceneStateBlueprint& InBlueprint, const FGuid& InTargetStructId, TArray<TInstancedStruct<FSceneStateBindingDesc>>& OutBindingDescs);
+
 	/**
 	 * Finds the state machine graph with id that matches the given struct id.
 	 * @param InBlueprint the blueprint to look into
diff --git a/Engine/Plugins/VirtualProduction/SceneState/Source/SceneStateBlueprint/Public/SceneStateBlueprint.h b/Engine/Plugins/VirtualProduction/SceneState/Source/SceneStateBlueprint/Public/SceneStateBlueprint.h
--- a/Engine/Plugins/VirtualProduction/SceneState/Source/SceneStateBlueprint/Public/SceneStateBlueprint.h
+++ b/Engine/Plugins/VirtualProduction/SceneState/Source/SceneStateBlueprint/Public/SceneStateBlueprint.h
@@ -37,6 +37,16 @@ public:
 
 	UE_API FSceneStateBindingDesc CreateRootBinding() const;
 
+	struct FGetBindableStructsParams
+	{
+		/** Struct id to look bindable structs for */
+		FGuid TargetStructId;
+		/** Whether to include the binding functions bound to the struct. Useful for compiling, but not */
+		bool bIncludeFunctions = false;
+	};
+	/** Gathers the binding structs for a given struct id. */
+	UE_API void GetBindingStructs(const FGetBindableStructsParams& InParams, TArray<TInstancedStruct<FPropertyBindingBindableStructDescriptor>>& OutStructDescs) const;
+
 	//~ Begin ISceneStateBindingCollectionOwner
 	UE_API virtual FSceneStateBindingCollection& GetBindingCollection() override;
 	UE_API virtual const FSceneStateBindingCollection& GetBindingCollection() const override;
diff --git a/Engine/Plugins/VirtualProduction/SceneState/Source/SceneStateBlueprintEditor/Private/Compiler/SceneStateBindingCompiler.cpp b/Engine/Plugins/VirtualProduction/SceneState/Source/SceneStateBlueprintEditor/Private/Compiler/SceneStateBindingCompiler.cpp
--- a/Engine/Plugins/VirtualProduction/SceneState/Source/SceneStateBlueprintEditor/Private/Compiler/SceneStateBindingCompiler.cpp
+++ b/Engine/Plugins/VirtualProduction/SceneState/Source/SceneStateBlueprintEditor/Private/Compiler/SceneStateBindingCompiler.cpp
@@ -257,10 +257,16 @@ void FBindingCompiler::AddFunctionBindingDescs()
 
 bool FBindingCompiler::ValidateBinding(const FSceneStateBinding& InBinding) const
 {
+	const USceneStateBlueprint::FGetBindableStructsParams Params
+		{
+			.TargetStructId = InBinding.GetTargetPath().GetStructID(),
+			.bIncludeFunctions = true,
+		};
+
 	// Source must be accessible by the target struct.
 	// This mismatch could happen if copying an object bound to a scoped parameter and pasting it outside such scope
 	TArray<TInstancedStruct<FPropertyBindingBindableStructDescriptor>> AccessibleStructs;
-	Blueprint->GetBindableStructs(InBinding.GetTargetPath().GetStructID(), AccessibleStructs);
+	Blueprint->GetBindingStructs(Params, AccessibleStructs);
 
 	const FGuid& SourceStructId = InBinding.GetSourcePath().GetStructID();
 	const bool bSourceAccessible = AccessibleStructs.ContainsByPredicate(
diff --git a/Engine/Plugins/VirtualProduction/SceneState/Source/SceneStateBlueprintEditor/Private/DetailsView/SceneStateCachedBindingData.cpp b/Engine/Plugins/VirtualProduction/SceneState/Source/SceneStateBlueprintEditor/Private/DetailsView/SceneStateCachedBindingData.cpp
--- a/Engine/Plugins/VirtualProduction/SceneState/Source/SceneStateBlueprintEditor/Private/DetailsView/SceneStateCachedBindingData.cpp
+++ b/Engine/Plugins/VirtualProduction/SceneState/Source/SceneStateBlueprintEditor/Private/DetailsView/SceneStateCachedBindingData.cpp
@@ -365,6 +365,55 @@ bool FSceneStateCachedBindingData::DeterminePropertiesCompatibilityInternal(cons
 	return false;
 }
 
+bool FSceneStateCachedBindingData::GetPropertyFunctionText(FConstStructView InPropertyFunctionStructView, FText& OutText) const
+{
+	// Todo: consider having each function set their own descriptions while also being able to determine binding names (via passing the binding owner)
+	const FSceneStateBindingFunction* BindingFunction = InPropertyFunctionStructView.GetPtr<const FSceneStateBindingFunction>();
+	if (BindingFunction && BindingFunction->Function.GetScriptStruct())
+	{
+		OutText = BindingFunction->Function.GetScriptStruct()->GetDisplayNameText();
+		return true;
+	}
+	return FCachedBindingData::GetPropertyFunctionText(InPropertyFunctionStructView, OutText);
+}
+
+bool FSceneStateCachedBindingData::GetPropertyFunctionTooltipText(FConstStructView InPropertyFunctionStructView, FText& OutText) const
+{
+	const FSceneStateBindingFunction* BindingFunction = InPropertyFunctionStructView.GetPtr<const FSceneStateBindingFunction>();
+	if (BindingFunction && BindingFunction->Function.GetScriptStruct())
+	{
+		OutText = BindingFunction->Function.GetScriptStruct()->GetToolTipText();
+		return true;
+	}
+	return Super::GetPropertyFunctionTooltipText(InPropertyFunctionStructView, OutText);
+}
+
+bool FSceneStateCachedBindingData::GetPropertyFunctionIconColor(FConstStructView InPropertyFunctionStructView, FLinearColor& OutColor) const
+{
+	if (const FSceneStateBindingFunction* BindingFunction = InPropertyFunctionStructView.GetPtr<const FSceneStateBindingFunction>())
+	{
+		if (const FProperty* OutputProperty = UE::SceneState::Editor::FindSingleOutputProperty(BindingFunction->FunctionInstance.GetScriptStruct()))
+		{
+			const UEdGraphSchema_K2* Schema = GetDefault<UEdGraphSchema_K2>();
+			check(Schema);
+
+			FEdGraphPinType PinType;
+			if (Schema->ConvertPropertyToPinType(OutputProperty, PinType))
+			{
+				OutColor = Schema->GetPinTypeColor(PinType);
+				return true;
+			}
+		}
+	}
+	return false;
+}
+
+bool FSceneStateCachedBindingData::GetPropertyFunctionImage(FConstStructView InPropertyFunctionStructView, const FSlateBrush*& OutImage) const
+{
+	OutImage = FAppStyle::GetBrush(TEXT("GraphEditor.Function_16x"));
+	return true;
+}
+
 } // UE::SceneState::Editor
 
 #undef LOCTEXT_NAMESPACE
diff --git a/Engine/Plugins/VirtualProduction/SceneState/Source/SceneStateBlueprintEditor/Private/DetailsView/SceneStateCachedBindingData.h b/Engine/Plugins/VirtualProduction/SceneState/Source/SceneStateBlueprintEditor/Private/DetailsView/SceneStateCachedBindingData.h
--- a/Engine/Plugins/VirtualProduction/SceneState/Source/SceneStateBlueprintEditor/Private/DetailsView/SceneStateCachedBindingData.h
+++ b/Engine/Plugins/VirtualProduction/SceneState/Source/SceneStateBlueprintEditor/Private/DetailsView/SceneStateCachedBindingData.h
@@ -38,6 +38,10 @@ struct FSceneStateCachedBindingData : UE::PropertyBinding::FCachedBindingData, T
 	virtual bool AddBindingInternal(TConstStructView<FPropertyBindingBindableStructDescriptor> InDescriptor, FPropertyBindingPath& InOutSourcePath, const FPropertyBindingPath& InTargetPath) override;
 	virtual void AddPropertyInfoOverride(const FProperty& InProperty, TArray<TSharedPtr<const UE::PropertyBinding::FPropertyInfoOverride>>& OutPropertyInfoOverrides) const override;
 	virtual bool DeterminePropertiesCompatibilityInternal(const FProperty* InSourceProperty,const FProperty* InTargetProperty,const void* InSourcePropertyValue,const void* InTargetPropertyValue,bool& bOutAreCompatible) const override;
+	virtual bool GetPropertyFunctionText(FConstStructView InPropertyFunctionStructView, FText& OutText) const;
+	virtual bool GetPropertyFunctionTooltipText(FConstStructView InPropertyFunctionStructView, FText& OutText) const;
+	virtual bool GetPropertyFunctionIconColor(FConstStructView InPropertyFunctionStructView, FLinearColor& OutColor) const;
+	virtual bool GetPropertyFunctionImage(FConstStructView InPropertyFunctionStructView, const FSlateBrush*& OutImage) const;
 	//~ End UE::PropertyBinding::FCachedBindingData
 };
 
diff --git a/Engine/Plugins/VirtualProduction/SceneState/Source/SceneStateBlueprintEditor/Private/DetailsView/SceneStatePropertyBindingExtension.cpp b/Engine/Plugins/VirtualProduction/SceneState/Source/SceneStateBlueprintEditor/Private/DetailsView/SceneStatePropertyBindingExtension.cpp
--- a/Engine/Plugins/VirtualProduction/SceneState/Source/SceneStateBlueprintEditor/Private/DetailsView/SceneStatePropertyBindingExtension.cpp
+++ b/Engine/Plugins/VirtualProduction/SceneState/Source/SceneStateBlueprintEditor/Private/DetailsView/SceneStatePropertyBindingExtension.cpp
@@ -73,6 +73,7 @@ bool FBindingExtension::GetPromotionToParameterOverrideInternal(const FProperty&
 void FBindingExtension::UpdateContextStruct(TConstStructView<FPropertyBindingBindableStructDescriptor> InStructDesc, FBindingContextStruct& InOutContextStruct, TMap<FString, FText>& InOutSectionNames) const
 {
 	const FSceneStateBindingDesc& BindingDesc = InStructDesc.Get<FSceneStateBindingDesc>();
+	InOutContextStruct.Section = BindingDesc.Section;
 
 	// Function overrides it's struct's icon color.
 	if (BindingDesc.DataHandle.IsDataType(ESceneStateDataType::Function))
