Motion Design: small fixes to function categories + added boolean and concatenate string functions

diff --git a/Engine/Plugins/VirtualProduction/SceneState/Source/SceneStateTasks/Private/Functions/SceneStateBooleanFunctions.cpp b/Engine/Plugins/VirtualProduction/SceneState/Source/SceneStateTasks/Private/Functions/SceneStateBooleanFunctions.cpp
new file mode 100644
--- /dev/null
+++ b/Engine/Plugins/VirtualProduction/SceneState/Source/SceneStateTasks/Private/Functions/SceneStateBooleanFunctions.cpp
@@ -0,0 +1,41 @@
+// Copyright Epic Games, Inc. All Rights Reserved.
+
+#include "Functions/SceneStateBooleanFunctions.h"
+
+#if WITH_EDITOR
+const UScriptStruct* FSceneStateBooleanFunction::OnGetFunctionDataType() const
+{
+	return FInstanceDataType::StaticStruct();
+}
+#endif
+
+void FSceneStateBooleanAndFunction::OnExecute(const FSceneStateExecutionContext& InContext, FStructView InFunctionInstance) const
+{
+	FInstanceDataType& InstanceData = InFunctionInstance.Get<FInstanceDataType>();
+	InstanceData.bOutput = InstanceData.bLeft && InstanceData.bRight;
+}
+
+void FSceneStateBooleanOrFunction::OnExecute(const FSceneStateExecutionContext& InContext, FStructView InFunctionInstance) const
+{
+	FInstanceDataType& InstanceData = InFunctionInstance.Get<FInstanceDataType>();
+	InstanceData.bOutput = InstanceData.bLeft || InstanceData.bRight;
+}
+
+void FSceneStateBooleanXorFunction::OnExecute(const FSceneStateExecutionContext& InContext, FStructView InFunctionInstance) const
+{
+	FInstanceDataType& InstanceData = InFunctionInstance.Get<FInstanceDataType>();
+	InstanceData.bOutput = InstanceData.bLeft ^ InstanceData.bRight;
+}
+
+#if WITH_EDITOR
+const UScriptStruct* FSceneStateBooleanNotFunction::OnGetFunctionDataType() const
+{
+	return FInstanceDataType::StaticStruct();
+}
+#endif
+
+void FSceneStateBooleanNotFunction::OnExecute(const FSceneStateExecutionContext& InContext, FStructView InFunctionInstance) const
+{
+	FInstanceDataType& InstanceData = InFunctionInstance.Get<FInstanceDataType>();
+	InstanceData.bOutput = !InstanceData.bInput;
+}
diff --git a/Engine/Plugins/VirtualProduction/SceneState/Source/SceneStateTasks/Private/Functions/SceneStateStringFunctions.cpp b/Engine/Plugins/VirtualProduction/SceneState/Source/SceneStateTasks/Private/Functions/SceneStateStringFunctions.cpp
new file mode 100644
--- /dev/null
+++ b/Engine/Plugins/VirtualProduction/SceneState/Source/SceneStateTasks/Private/Functions/SceneStateStringFunctions.cpp
@@ -0,0 +1,16 @@
+// Copyright Epic Games, Inc. All Rights Reserved.
+
+#include "Functions/SceneStateStringFunctions.h"
+
+#if WITH_EDITOR
+const UScriptStruct* FSceneStateConcatenateStringFunction::OnGetFunctionDataType() const
+{
+	return FInstanceDataType::StaticStruct();
+}
+#endif
+
+void FSceneStateConcatenateStringFunction::OnExecute(const FSceneStateExecutionContext& InContext, FStructView InFunctionInstance) const
+{
+	FInstanceDataType& InstanceData = InFunctionInstance.Get<FInstanceDataType>();
+	InstanceData.Output = InstanceData.Left + InstanceData.Right;
+}
diff --git a/Engine/Plugins/VirtualProduction/SceneState/Source/SceneStateTasks/Public/Functions/SceneStateBooleanFunctions.h b/Engine/Plugins/VirtualProduction/SceneState/Source/SceneStateTasks/Public/Functions/SceneStateBooleanFunctions.h
new file mode 100644
--- /dev/null
+++ b/Engine/Plugins/VirtualProduction/SceneState/Source/SceneStateTasks/Public/Functions/SceneStateBooleanFunctions.h
@@ -0,0 +1,104 @@
+// Copyright Epic Games, Inc. All Rights Reserved.
+
+#pragma once
+
+#include "Functions/SceneStateFunction.h"
+#include "SceneStateBooleanFunctions.generated.h"
+
+/** Function instance for all boolean operation functions involving two operands, left and right, and a result */
+USTRUCT()
+struct FSceneStateBooleanOpFunctionInstance
+{
+	GENERATED_BODY()
+
+	UPROPERTY(EditAnywhere, Category="Math")
+	bool bLeft = false;
+
+	UPROPERTY(EditAnywhere, Category="Math")
+	bool bRight = false;
+
+	UPROPERTY(EditAnywhere, Category="Math", meta=(Output))
+	bool bOutput = false;
+};
+
+/** Base class for boolean operations where there are two operands, left and right, and a result */
+USTRUCT(meta=(Hidden))
+struct FSceneStateBooleanFunction : public FSceneStateFunction
+{
+	GENERATED_BODY()
+
+	using FInstanceDataType = FSceneStateBooleanOpFunctionInstance;
+
+protected:
+	//~ Begin FSceneStateFunction
+#if WITH_EDITOR
+	virtual const UScriptStruct* OnGetFunctionDataType() const override;
+#endif
+	//~ End FSceneStateFunction
+};
+
+/** Performs boolean AND (Left AND Right) */
+USTRUCT(DisplayName="And", Category="Math")
+struct FSceneStateBooleanAndFunction : public FSceneStateBooleanFunction
+{
+	GENERATED_BODY()
+
+protected:
+	//~ Begin FSceneStateFunction
+	virtual void OnExecute(const FSceneStateExecutionContext& InContext, FStructView InFunctionInstance) const override;
+	//~ End FSceneStateFunction
+};
+
+/** Performs boolean OR (Left AND Right) */
+USTRUCT(DisplayName="Or", Category="Math")
+struct FSceneStateBooleanOrFunction : public FSceneStateBooleanFunction
+{
+	GENERATED_BODY()
+
+protected:
+	//~ Begin FSceneStateFunction
+	virtual void OnExecute(const FSceneStateExecutionContext& InContext, FStructView InFunctionInstance) const override;
+	//~ End FSceneStateFunction
+};
+
+/** Performs boolean XOR (Left XOR Right) */
+USTRUCT(DisplayName="XOr", Category="Math")
+struct FSceneStateBooleanXorFunction : public FSceneStateBooleanFunction
+{
+	GENERATED_BODY()
+
+protected:
+	//~ Begin FSceneStateFunction
+	virtual void OnExecute(const FSceneStateExecutionContext& InContext, FStructView InFunctionInstance) const override;
+	//~ End FSceneStateFunction
+};
+
+/** Function instance the boolean NOT operation where only one input is required */
+USTRUCT()
+struct FSceneStateBooleanNotFunctionInstance
+{
+	GENERATED_BODY()
+
+	UPROPERTY(EditAnywhere, Category="Math")
+	bool bInput = false;
+
+	UPROPERTY(EditAnywhere, Category="Math", meta=(Output))
+	bool bOutput = false;
+};
+
+/** Performs boolean NOT */
+USTRUCT(DisplayName="Not", Category="Math")
+struct FSceneStateBooleanNotFunction : public FSceneStateBooleanFunction
+{
+	GENERATED_BODY()
+
+	using FInstanceDataType = FSceneStateBooleanNotFunctionInstance;
+
+protected:
+	//~ Begin FSceneStateFunction
+#if WITH_EDITOR
+	virtual const UScriptStruct* OnGetFunctionDataType() const override;
+#endif
+	virtual void OnExecute(const FSceneStateExecutionContext& InContext, FStructView InFunctionInstance) const override;
+	//~ End FSceneStateFunction
+};
diff --git a/Engine/Plugins/VirtualProduction/SceneState/Source/SceneStateTasks/Public/Functions/SceneStateFloatFunctions.h b/Engine/Plugins/VirtualProduction/SceneState/Source/SceneStateTasks/Public/Functions/SceneStateFloatFunctions.h
--- a/Engine/Plugins/VirtualProduction/SceneState/Source/SceneStateTasks/Public/Functions/SceneStateFloatFunctions.h
+++ b/Engine/Plugins/VirtualProduction/SceneState/Source/SceneStateTasks/Public/Functions/SceneStateFloatFunctions.h
@@ -7,17 +7,17 @@
 
 /** Function instance for all double-precision float operation functions involving two operands, left and right, and a result */
 USTRUCT()
-struct FSceneStateDoubleFunctionInstance
+struct FSceneStateDoubleOpFunctionInstance
 {
 	GENERATED_BODY()
 
-	UPROPERTY(EditAnywhere, Category="Core")
+	UPROPERTY(EditAnywhere, Category="Math")
 	double Left = 0.0;
 
-	UPROPERTY(EditAnywhere, Category="Core")
+	UPROPERTY(EditAnywhere, Category="Math")
 	double Right = 0.0;
 
-	UPROPERTY(EditAnywhere, Category="Core", meta=(Output))
+	UPROPERTY(EditAnywhere, Category="Math", meta=(Output))
 	double Output = 0.0;
 };
 
@@ -27,7 +27,7 @@ struct FSceneStateDoubleFunction : public FSceneStateFunction
 {
 	GENERATED_BODY()
 
-	using FInstanceDataType = FSceneStateDoubleFunctionInstance;
+	using FInstanceDataType = FSceneStateDoubleOpFunctionInstance;
 
 protected:
 	//~ Begin FSceneStateFunction
diff --git a/Engine/Plugins/VirtualProduction/SceneState/Source/SceneStateTasks/Public/Functions/SceneStateFromStringFunctions.h b/Engine/Plugins/VirtualProduction/SceneState/Source/SceneStateTasks/Public/Functions/SceneStateFromStringFunctions.h
--- a/Engine/Plugins/VirtualProduction/SceneState/Source/SceneStateTasks/Public/Functions/SceneStateFromStringFunctions.h
+++ b/Engine/Plugins/VirtualProduction/SceneState/Source/SceneStateTasks/Public/Functions/SceneStateFromStringFunctions.h
@@ -11,7 +11,7 @@ struct FSceneStateFromStringFunctionInstance
 {
 	GENERATED_BODY()
 
-	UPROPERTY(EditAnywhere, Category="Core")
+	UPROPERTY(EditAnywhere, Category="String")
 	FString String;
 };
 
@@ -23,11 +23,11 @@ struct FSceneStateTextFromStringFunctionInstance : public FSceneStateFromStringF
 {
 	GENERATED_BODY()
 
-	UPROPERTY(EditAnywhere, Category="Core", meta=(Output))
+	UPROPERTY(EditAnywhere, Category="String", meta=(Output))
 	FText Output;
 };
 
-USTRUCT(DisplayName="String to Text", Category="Core")
+USTRUCT(DisplayName="String to Text", Category="String")
 struct FSceneStateTextFromStringFunction : public FSceneStateFunction
 {
 	GENERATED_BODY()
@@ -54,11 +54,11 @@ struct FSceneStateNameFromStringFunctionInstance : public FSceneStateFromStringF
 {
 	GENERATED_BODY()
 
-	UPROPERTY(EditAnywhere, Category="Core", meta=(Output))
+	UPROPERTY(EditAnywhere, Category="String", meta=(Output))
 	FName Output;
 };
 
-USTRUCT(DisplayName="String to Name", Category="Core")
+USTRUCT(DisplayName="String to Name", Category="String")
 struct FSceneStateNameFromStringFunction : public FSceneStateFunction
 {
 	GENERATED_BODY()
diff --git a/Engine/Plugins/VirtualProduction/SceneState/Source/SceneStateTasks/Public/Functions/SceneStateIntegerFunctions.h b/Engine/Plugins/VirtualProduction/SceneState/Source/SceneStateTasks/Public/Functions/SceneStateIntegerFunctions.h
--- a/Engine/Plugins/VirtualProduction/SceneState/Source/SceneStateTasks/Public/Functions/SceneStateIntegerFunctions.h
+++ b/Engine/Plugins/VirtualProduction/SceneState/Source/SceneStateTasks/Public/Functions/SceneStateIntegerFunctions.h
@@ -7,17 +7,17 @@
 
 /** Function instance for all integer operation functions involving two operands, left and right, and an integer result */
 USTRUCT()
-struct FSceneStateIntegerFunctionInstance
+struct FSceneStateIntegerOpFunctionInstance
 {
 	GENERATED_BODY()
 
-	UPROPERTY(EditAnywhere, Category="Core")
+	UPROPERTY(EditAnywhere, Category="Math")
 	int32 Left = 0;
 
-	UPROPERTY(EditAnywhere, Category="Core")
+	UPROPERTY(EditAnywhere, Category="Math")
 	int32 Right = 0;
 
-	UPROPERTY(EditAnywhere, Category="Core", meta=(Output))
+	UPROPERTY(EditAnywhere, Category="Math", meta=(Output))
 	int32 Output = 0;
 };
 
@@ -27,7 +27,7 @@ struct FSceneStateIntegerFunction : public FSceneStateFunction
 {
 	GENERATED_BODY()
 
-	using FInstanceDataType = FSceneStateIntegerFunctionInstance;
+	using FInstanceDataType = FSceneStateIntegerOpFunctionInstance;
 
 protected:
 	//~ Begin FSceneStateFunction
diff --git a/Engine/Plugins/VirtualProduction/SceneState/Source/SceneStateTasks/Public/Functions/SceneStateStringFunctions.h b/Engine/Plugins/VirtualProduction/SceneState/Source/SceneStateTasks/Public/Functions/SceneStateStringFunctions.h
new file mode 100644
--- /dev/null
+++ b/Engine/Plugins/VirtualProduction/SceneState/Source/SceneStateTasks/Public/Functions/SceneStateStringFunctions.h
@@ -0,0 +1,38 @@
+// Copyright Epic Games, Inc. All Rights Reserved.
+
+#pragma once
+
+#include "Functions/SceneStateFunction.h"
+#include "SceneStateStringFunctions.generated.h"
+
+USTRUCT()
+struct FSceneStateConcatenateStringFunctionInstance
+{
+	GENERATED_BODY()
+
+	UPROPERTY(EditAnywhere, Category="String")
+	FString Left;
+
+	UPROPERTY(EditAnywhere, Category="String")
+	FString Right;
+
+	UPROPERTY(EditAnywhere, Category="String", meta=(Output))
+	FString Output;
+};
+
+/** Concatenates two strings (Left + Right) together */
+USTRUCT(DisplayName="Concatenate String", Category="String")
+struct FSceneStateConcatenateStringFunction : public FSceneStateFunction
+{
+	GENERATED_BODY()
+
+	using FInstanceDataType = FSceneStateConcatenateStringFunctionInstance;
+
+protected:
+	//~ Begin FSceneStateFunction
+#if WITH_EDITOR
+	virtual const UScriptStruct* OnGetFunctionDataType() const override;
+#endif
+	virtual void OnExecute(const FSceneStateExecutionContext& InContext, FStructView InFunctionInstance) const override;
+	//~ End FSceneStateFunction
+};
diff --git a/Engine/Plugins/VirtualProduction/SceneState/Source/SceneStateTasks/Public/Functions/SceneStateToStringFunctions.h b/Engine/Plugins/VirtualProduction/SceneState/Source/SceneStateTasks/Public/Functions/SceneStateToStringFunctions.h
--- a/Engine/Plugins/VirtualProduction/SceneState/Source/SceneStateTasks/Public/Functions/SceneStateToStringFunctions.h
+++ b/Engine/Plugins/VirtualProduction/SceneState/Source/SceneStateTasks/Public/Functions/SceneStateToStringFunctions.h
@@ -11,7 +11,7 @@ struct FSceneStateToStringFunctionInstance
 {
 	GENERATED_BODY()
 
-	UPROPERTY(EditAnywhere, Category="Core", meta=(Output))
+	UPROPERTY(EditAnywhere, Category="String", meta=(Output))
 	FString Output;
 };
 
@@ -23,11 +23,11 @@ struct FSceneStateTextToStringFunctionInstance : public FSceneStateToStringFunct
 {
 	GENERATED_BODY()
 
-	UPROPERTY(EditAnywhere, Category="Core")
+	UPROPERTY(EditAnywhere, Category="String")
 	FText Value;
 };
 
-USTRUCT(DisplayName="Text to String", Category="Core")
+USTRUCT(DisplayName="Text to String", Category="String")
 struct FSceneStateTextToStringFunction : public FSceneStateFunction
 {
 	GENERATED_BODY()
@@ -54,11 +54,11 @@ struct FSceneStateNameToStringFunctionInstance : public FSceneStateToStringFunct
 {
 	GENERATED_BODY()
 
-	UPROPERTY(EditAnywhere, Category="Core")
+	UPROPERTY(EditAnywhere, Category="String")
 	FName Value;
 };
 
-USTRUCT(DisplayName="Name to String", Category="Core")
+USTRUCT(DisplayName="Name to String", Category="String")
 struct FSceneStateNameToStringFunction : public FSceneStateFunction
 {
 	GENERATED_BODY()
@@ -85,11 +85,11 @@ struct FSceneStateIntegerToStringFunctionInstance : public FSceneStateToStringFu
 {
 	GENERATED_BODY()
 
-	UPROPERTY(EditAnywhere, Category="Core")
+	UPROPERTY(EditAnywhere, Category="String")
 	int32 Value = 0;
 };
 
-USTRUCT(DisplayName="Integer to String", Category="Core")
+USTRUCT(DisplayName="Integer to String", Category="String")
 struct FSceneStateIntegerToStringFunction : public FSceneStateFunction
 {
 	GENERATED_BODY()
@@ -116,11 +116,11 @@ struct FSceneStateDoubleToStringFunctionInstance : public FSceneStateToStringFun
 {
 	GENERATED_BODY()
 
-	UPROPERTY(EditAnywhere, Category="Core")
+	UPROPERTY(EditAnywhere, Category="String")
 	double Value = 0.0;
 };
 
-USTRUCT(DisplayName="Float to String", Category="Core")
+USTRUCT(DisplayName="Float to String", Category="String")
 struct FSceneStateDoubleToStringFunction : public FSceneStateFunction
 {
 	GENERATED_BODY()
