Motion Design Scene State: moved the generated class scene state data to its own struct to decouple compilers and executors from the generated class.
This is done to allow to build unit tests without needing to deal with the generated class, and only use this template data struct.

#rnx

diff --git a/Engine/Plugins/VirtualProduction/SceneState/Source/SceneState/Private/Conduit/SceneStateConduit.cpp b/Engine/Plugins/VirtualProduction/SceneState/Source/SceneState/Private/Conduit/SceneStateConduit.cpp
--- a/Engine/Plugins/VirtualProduction/SceneState/Source/SceneState/Private/Conduit/SceneStateConduit.cpp
+++ b/Engine/Plugins/VirtualProduction/SceneState/Source/SceneState/Private/Conduit/SceneStateConduit.cpp
@@ -3,17 +3,17 @@
 #include "Conduit/SceneStateConduit.h"
 #include "Conduit/SceneStateConduitInstance.h"
 #include "Conduit/SceneStateConduitLink.h"
+#include "SceneState.h"
 #include "SceneStateExecutionContext.h"
-#include "SceneStateGeneratedClass.h"
-#include "SceneStateObject.h"
+#include "Transition/SceneStateTransition.h"
 #include "Transition/SceneStateTransitionEvaluation.h"
 #include "Transition/SceneStateTransitionResult.h"
 
-void FSceneStateConduit::Link(const FSceneStateConduitLink& InConduitLink, USceneStateGeneratedClass* InGeneratedClass)
+void FSceneStateConduit::Link(const FSceneStateConduitLink& InConduitLink, TNotNull<const UClass*> InOwnerClass)
 {
 	if (!InConduitLink.ResultPropertyName.IsNone())
 	{
-		ResultProperty = CastField<FStructProperty>(InGeneratedClass->FindPropertyByName(InConduitLink.ResultPropertyName));
+		ResultProperty = CastField<FStructProperty>(InOwnerClass->FindPropertyByName(InConduitLink.ResultPropertyName));
 		check(ResultProperty);
 	}
 	else
@@ -23,7 +23,7 @@ void FSceneStateConduit::Link(const FSceneStateConduitLink& InConduitLink, UScen
 
 	if (!InConduitLink.EventName.IsNone())
 	{
-		EvaluationEvent = InGeneratedClass->FindFunctionByName(InConduitLink.EventName);
+		EvaluationEvent = InOwnerClass->FindFunctionByName(InConduitLink.EventName);
 		check(EvaluationEvent);
 	}
 	else
@@ -79,11 +79,15 @@ bool FSceneStateConduit::ProcessEvaluationEvent(const UE::SceneState::FTransitio
 
 	void* FunctionParams = nullptr;
 
-	USceneStateObject* RootState = InParams.ExecutionContext.GetRootState();
-	check(RootState);
-	RootState->ProcessEvent(EvaluationEvent, FunctionParams);
+	UObject* RootObject = InParams.ExecutionContext.GetRootObject();
+	if (!ensure(RootObject))
+	{
+		return false;
+	}
+
+	RootObject->ProcessEvent(EvaluationEvent, FunctionParams);
 
-	const FSceneStateTransitionResult* Result = ResultProperty->ContainerPtrToValuePtr<FSceneStateTransitionResult>(RootState);
+	const FSceneStateTransitionResult* Result = ResultProperty->ContainerPtrToValuePtr<FSceneStateTransitionResult>(RootObject);
 	check(Result);
 	return Result->bCanTransition;
 }
diff --git a/Engine/Plugins/VirtualProduction/SceneState/Source/SceneState/Private/SceneState.cpp b/Engine/Plugins/VirtualProduction/SceneState/Source/SceneState/Private/SceneState.cpp
--- a/Engine/Plugins/VirtualProduction/SceneState/Source/SceneState/Private/SceneState.cpp
+++ b/Engine/Plugins/VirtualProduction/SceneState/Source/SceneState/Private/SceneState.cpp
@@ -244,7 +244,7 @@ void FSceneState::AllocateTaskInstances(const FSceneStateExecutionContext& InCon
 		*TaskInstanceContainer = InTemplateTaskInstances;
 
 		// Instance each Template Object into the Instance data
-		InstanceTaskObjects(InContext.GetRootState()
+		InstanceTaskObjects(InContext.GetRootObject()
 			, UE::SceneState::GetStructViews(*TaskInstanceContainer)
 			, InTemplateTaskInstances
 			, [](FObjectDuplicationParameters& InParams)->UObject*
diff --git a/Engine/Plugins/VirtualProduction/SceneState/Source/SceneState/Private/SceneStateExecutionContext.cpp b/Engine/Plugins/VirtualProduction/SceneState/Source/SceneState/Private/SceneStateExecutionContext.cpp
--- a/Engine/Plugins/VirtualProduction/SceneState/Source/SceneState/Private/SceneStateExecutionContext.cpp
+++ b/Engine/Plugins/VirtualProduction/SceneState/Source/SceneState/Private/SceneStateExecutionContext.cpp
@@ -6,11 +6,12 @@
 #include "PropertyBindingTypes.h"
 #include "SceneState.h"
 #include "SceneStateExecutionContextRegistry.h"
-#include "SceneStateGeneratedClass.h"
 #include "SceneStateLog.h"
 #include "SceneStateMachine.h"
 #include "SceneStateObject.h"
+#include "SceneStateTemplateData.h"
 #include "SceneStateUtils.h"
+#include "Tasks/SceneStateTask.h"
 #include "Tasks/SceneStateTaskBindingExtension.h"
 #include "Tasks/SceneStateTaskInstance.h"
 
@@ -21,73 +22,94 @@ FSceneStateExecutionContext::~FSceneStateExecutionContext()
 	Reset();
 }
 
-TWeakPtr<const UE::SceneState::FExecutionContextRegistry> FSceneStateExecutionContext::GetContextRegistry() const
+const FSceneStateTemplateData& FSceneStateExecutionContext::GetSharedTemplateData() const
 {
-	return ContextRegistryWeak;
+	return SharedTemplateData.Get<const FSceneStateTemplateData>();
 }
 
-const FSceneStateBindingCollection* FSceneStateExecutionContext::GetBindingCollection() const
+TSharedPtr<const UE::SceneState::FExecutionContextRegistry> FSceneStateExecutionContext::GetContextRegistry() const
 {
-	if (GeneratedClass)
-	{
-		return &GeneratedClass->GetBindingCollection();
-	}
-	return nullptr;
+	return ContextRegistryWeak.Pin();
 }
 
-FString FSceneStateExecutionContext::GetExecutionContextName() const
+const FSceneStateBindingCollection& FSceneStateExecutionContext::GetBindingCollection() const
 {
-	if (RootState)
-	{
-		return RootState->GetContextName();
-	}
-	return FString();
+	return GetSharedTemplateData().BindingCollection;
+}
+
+const FString& FSceneStateExecutionContext::GetExecutionContextName() const
+{
+	return ContextName;
+}
+
+UObject* FSceneStateExecutionContext::GetRootObject() const
+{
+	return RootObject;
 }
 
 UObject* FSceneStateExecutionContext::GetContextObject() const
 {
-	if (RootState)
-	{
-		return RootState->GetContextObject();
-	}
-	return nullptr;
+	return ContextObjectWeak.Get();
 }
 
 USceneStateEventStream* FSceneStateExecutionContext::GetEventStream() const
 {
-	if (RootState)
-	{
-		return RootState->GetEventStream();
-	}
-	return nullptr;
+	return EventStreamWeak.Get();
 }
 
 void FSceneStateExecutionContext::Setup(const FSceneStateExecutionContext& InParentContext)
 {
-	Setup(InParentContext.GetRootState());
-	check(RootState);
+	const FSetupParams SetupParams
+	{
+		.SharedTemplateData = InParentContext.SharedTemplateData,
+		.RootObject = InParentContext.GetRootObject(),
+		.ContextObject = InParentContext.GetContextObject(),
+		.EventStream = InParentContext.GetEventStream(),
+		.ContextRegistry = InParentContext.ContextRegistryWeak.Pin(),
+		.ContextName = InParentContext.GetExecutionContextName(),
+	};
+	Setup(SetupParams);
+
+	if (SetupParams.ContextRegistry.IsValid())
+	{
+		ParentHandle = SetupParams.ContextRegistry->FindHandle(InParentContext);
+	}
+}
 
-	const TSharedRef<UE::SceneState::FExecutionContextRegistry> ContextRegistry = RootState->GetContextRegistry();
-	ParentHandle = ContextRegistry->FindHandle(InParentContext);
+void FSceneStateExecutionContext::Setup(TNotNull<USceneStateObject*> InRootState)
+{
+	const FSetupParams SetupParams
+	{
+		.SharedTemplateData = InRootState->GetSharedTemplateData(),
+		.RootObject = InRootState,
+		.ContextObject = InRootState->GetContextObject(),
+		.EventStream = InRootState->GetEventStream(),
+		.ContextRegistry = InRootState->GetContextRegistry(),
+		.ContextName = InRootState->GetContextName(),
+	};
+	Setup(SetupParams);
 }
 
-void FSceneStateExecutionContext::Setup(USceneStateObject* InRootState)
+void FSceneStateExecutionContext::Setup(const FSetupParams& InParams)
 {
 	Reset();
 
-	check(InRootState);
-	RootState = InRootState;
-	GeneratedClass = Cast<USceneStateGeneratedClass>(RootState->GetClass());
+	SharedTemplateData = InParams.SharedTemplateData;
+	RootObject = InParams.RootObject;
+	ContextObjectWeak = InParams.ContextObject;
+	EventStreamWeak = InParams.EventStream;
+	ContextRegistryWeak = InParams.ContextRegistry;
+	ContextName = InParams.ContextName;
 
-	const TSharedRef<UE::SceneState::FExecutionContextRegistry> ContextRegistry = InRootState->GetContextRegistry();
-	ContextRegistry->RegisterContext(this);
-	ContextRegistryWeak = ContextRegistry;
+	if (InParams.ContextRegistry.IsValid())
+	{
+		InParams.ContextRegistry->RegisterContext(this);
+	}
 }
 
 void FSceneStateExecutionContext::Reset()
 {
-	RootState = nullptr;
-	GeneratedClass = nullptr;
+	SharedTemplateData.Reset();
 
 	StateInstances.Reset();
 	StateMachineInstances.Reset();
@@ -122,7 +144,7 @@ FPropertyBindingDataView FSceneStateExecutionContext::FindDataView(const FSceneS
 	switch (DataType)
 	{
 	case ESceneStateDataType::Root:
-		return FPropertyBindingDataView(RootState);
+		return FPropertyBindingDataView(RootObject);
 
 	case ESceneStateDataType::Task:
 		return FindTaskInstance(DataIndex);
@@ -184,7 +206,7 @@ FPropertyBindingDataView FSceneStateExecutionContext::FindDataView(const FSceneS
 void FSceneStateExecutionContext::ForEachTask(const FSceneState& InState, TFunctionRef<UE::SceneState::EIterationResult(const FSceneStateTask&, FStructView)> InCallable) const
 {
 	const FSceneStateRange TaskRange = InState.GetTaskRange();
-	if (TaskRange.Count == 0 || !GeneratedClass)
+	if (TaskRange.Count == 0)
 	{
 		return;
 	}
@@ -195,7 +217,9 @@ void FSceneStateExecutionContext::ForEachTask(const FSceneState& InState, TFunct
 		return;
 	}
 
-	const FInstancedStructContainer& Tasks = GeneratedClass->GetTasks();
+	const FSceneStateTemplateData& TemplateData = GetSharedTemplateData();
+
+	const FInstancedStructContainer& Tasks = TemplateData.Tasks;
 
 	checkfSlow(Tasks.IsValidIndex(TaskRange.Index) && Tasks.IsValidIndex(TaskRange.GetLastIndex())
 		, TEXT("[Context: %s] State Task Range [%d, %d] is out of bounds! (Total Task Count: %d)")
@@ -225,22 +249,25 @@ void FSceneStateExecutionContext::ForEachTask(const FSceneState& InState, TFunct
 TArray<FConstStructView> FSceneStateExecutionContext::GetTemplateTaskInstances(const FSceneState& InState) const
 {
 	const FSceneStateRange TaskRange = InState.GetTaskRange();
-	if (TaskRange.Count == 0 || !GeneratedClass)
+	if (TaskRange.Count == 0)
 	{
 		return {};
 	}
-	return UE::SceneState::GetConstStructViews(GeneratedClass->GetTaskInstances(), TaskRange);
+
+	const FSceneStateTemplateData& TemplateData = GetSharedTemplateData();
+	return UE::SceneState::GetConstStructViews(TemplateData.TaskInstances, TaskRange);
 }
 
 TConstArrayView<uint16> FSceneStateExecutionContext::GetTaskPrerequisites(const FSceneStateTask& InTask) const
 {
 	const FSceneStateRange TaskPrerequisiteRange = InTask.GetPrerequisiteRange();
-	if (TaskPrerequisiteRange.Count == 0 || !GeneratedClass)
+	if (TaskPrerequisiteRange.Count == 0)
 	{
 		return {};
 	}
 
-	const TConstArrayView<uint16> TaskPrerequisitesView = GeneratedClass->GetTaskPrerequisites();
+	const FSceneStateTemplateData& TemplateData = GetSharedTemplateData();
+	const TConstArrayView<uint16> TaskPrerequisitesView = TemplateData.TaskPrerequisites;
 
 	checkfSlow(TaskPrerequisitesView.IsValidIndex(TaskPrerequisiteRange.Index) && TaskPrerequisitesView.IsValidIndex(TaskPrerequisiteRange.GetLastIndex())
 		, TEXT("[Context: %s] Task Prerequisite Range [%d, %d] is out of bounds! (Total Prerequisite Count: %d)")
@@ -255,12 +282,13 @@ TConstArrayView<uint16> FSceneStateExecutionContext::GetTaskPrerequisites(const
 TConstArrayView<FSceneStateTransition> FSceneStateExecutionContext::GetTransitions(const FSceneState& InState) const
 {
 	const FSceneStateRange TransitionRange = InState.GetTransitionRange();
-	if (TransitionRange.Count == 0 || !GeneratedClass)
+	if (TransitionRange.Count == 0)
 	{
 		return {};
 	}
 
-	const TConstArrayView<FSceneStateTransition> TransitionsView = GeneratedClass->GetTransitions();
+	const FSceneStateTemplateData& TemplateData = GetSharedTemplateData();
+	const TConstArrayView<FSceneStateTransition> TransitionsView = TemplateData.Transitions;
 
 	checkfSlow(TransitionsView.IsValidIndex(TransitionRange.Index) && TransitionsView.IsValidIndex(TransitionRange.GetLastIndex())
 		, TEXT("[Context: %s] State Transition Range [%d, %d] is out of bounds! (Total Transition Count: %d)")
@@ -275,12 +303,13 @@ TConstArrayView<FSceneStateTransition> FSceneStateExecutionContext::GetTransitio
 TConstArrayView<FSceneStateTransition> FSceneStateExecutionContext::GetTransitions(const FSceneStateConduit& InConduit) const
 {
 	const FSceneStateRange TransitionRange = InConduit.GetTransitionRange();
-	if (TransitionRange.Count == 0 || !GeneratedClass)
+	if (TransitionRange.Count == 0)
 	{
 		return {};
 	}
 
-	const TConstArrayView<FSceneStateTransition> TransitionsView = GeneratedClass->GetTransitions();
+	const FSceneStateTemplateData& TemplateData = GetSharedTemplateData();
+	const TConstArrayView<FSceneStateTransition> TransitionsView = TemplateData.Transitions;
 
 	checkfSlow(TransitionsView.IsValidIndex(TransitionRange.Index) && TransitionsView.IsValidIndex(TransitionRange.GetLastIndex())
 		, TEXT("[Context: %s] Conduit Transition Range [%d, %d] is out of bounds! (Total Transition Count: %d)")
@@ -294,23 +323,20 @@ TConstArrayView<FSceneStateTransition> FSceneStateExecutionContext::GetTransitio
 
 FInstancedPropertyBag FSceneStateExecutionContext::GetTemplateTransitionParameter(const FSceneStateTransition& InTransition) const
 {
-	uint16 TransitionIndex;
-	if (GetTransitionIndex(InTransition, TransitionIndex))
-	{
-		return GeneratedClass->GetTransitionParameters(TransitionIndex);
-	}
-	return FInstancedPropertyBag();
+	const FSceneStateTemplateData& TemplateData = GetSharedTemplateData();
+	return TemplateData.TransitionParameters.FindRef(GetTransitionIndex(TemplateData, InTransition));
 }
 
 TConstArrayView<FSceneStateMachine> FSceneStateExecutionContext::GetStateMachines(const FSceneState& InState) const
 {
 	const FSceneStateRange StateMachineRange = InState.GetStateMachineRange();
-	if (StateMachineRange.Count == 0 || !GeneratedClass)
+	if (StateMachineRange.Count == 0)
 	{
 		return {};
 	}
 
-	const TConstArrayView<FSceneStateMachine> StateMachinesView = GeneratedClass->GetStateMachines();
+	const FSceneStateTemplateData& TemplateData = GetSharedTemplateData();
+	const TConstArrayView<FSceneStateMachine> StateMachinesView = TemplateData.StateMachines;
 
 	checkfSlow(StateMachinesView.IsValidIndex(StateMachineRange.Index) && StateMachinesView.IsValidIndex(StateMachineRange.GetLastIndex())
 		, TEXT("[Context: %s] State State Machine Range [%d, %d] is out of bounds! (Total State Machine Count: %d)")
@@ -325,12 +351,13 @@ TConstArrayView<FSceneStateMachine> FSceneStateExecutionContext::GetStateMachine
 TConstArrayView<FSceneStateEventHandler> FSceneStateExecutionContext::GetEventHandlers(const FSceneState& InState) const
 {
 	const FSceneStateRange EventHandlerRange = InState.GetEventHandlerRange();
-	if (EventHandlerRange.Count == 0 || !GeneratedClass)
+	if (EventHandlerRange.Count == 0)
 	{
 		return {};
 	}
 
-	const TConstArrayView<FSceneStateEventHandler> EventHandlersView = GeneratedClass->GetEventHandlers();
+	const FSceneStateTemplateData& TemplateData = GetSharedTemplateData();
+	const TConstArrayView<FSceneStateEventHandler> EventHandlersView = TemplateData.EventHandlers;
 
 	checkfSlow(EventHandlersView.IsValidIndex(EventHandlerRange.Index) && EventHandlersView.IsValidIndex(EventHandlerRange.GetLastIndex())
 		, TEXT("[Context: %s] State Event Handler Range [%d, %d] is out of bounds! (Total Event Handler Count: %d)")
@@ -345,19 +372,10 @@ TConstArrayView<FSceneStateEventHandler> FSceneStateExecutionContext::GetEventHa
 #if WITH_EDITOR
 const FSceneStateMetadata* FSceneStateExecutionContext::GetStateMetadata(const FSceneState& InState) const
 {
-	if (!GeneratedClass)
-	{
-		return nullptr;
-	}
-
-	uint16 StateIndex;
-	if (!GetStateIndex(InState, StateIndex))
-	{
-		UE_LOG(LogSceneState, Error, TEXT("[Context: %s] State Index was not able to be retrieved. State Data is invalid!"), *GetExecutionContextName());
-		return nullptr;
-	}
+	const FSceneStateTemplateData& TemplateData = GetSharedTemplateData();
 
-	const TConstArrayView<FSceneStateMetadata> StateMetadataView = GeneratedClass->GetStateMetadata();
+	const uint16 StateIndex = GetStateIndex(TemplateData, InState);
+	const TConstArrayView<FSceneStateMetadata> StateMetadataView = TemplateData.StateMetadata;
 
 	if (!StateMetadataView.IsValidIndex(StateIndex))
 	{
@@ -375,21 +393,12 @@ const FSceneStateMetadata* FSceneStateExecutionContext::GetStateMetadata(const F
 
 const FSceneStateMachine* FSceneStateExecutionContext::GetStateMachine(const FGuid& InStateMachineId) const
 {
-	if (GeneratedClass)
-	{
-		return GeneratedClass->FindStateMachine(InStateMachineId);
-	}
-	return nullptr;
+	return GetSharedTemplateData().FindStateMachine(InStateMachineId);
 }
 
 const FSceneStateMachine* FSceneStateExecutionContext::GetStateMachine(uint16 InAbsoluteIndex) const
 {
-	if (!GeneratedClass)
-	{
-		return nullptr;
-	}
-
-	const TConstArrayView<FSceneStateMachine> StateMachinesView = GeneratedClass->GetStateMachines();
+	const TConstArrayView<FSceneStateMachine> StateMachinesView = GetSharedTemplateData().StateMachines;
 
 	checkfSlow(StateMachinesView.IsValidIndex(InAbsoluteIndex)
 		, TEXT("[Context: %s] Absolute Index %d is out of bounds! (Total State Machine Count: %d)")
@@ -402,12 +411,11 @@ const FSceneStateMachine* FSceneStateExecutionContext::GetStateMachine(uint16 In
 
 const FSceneStateMachine* FSceneStateExecutionContext::GetContextStateMachine() const
 {
-	if (StateMachineExecutionStack.IsEmpty() || !GeneratedClass)
+	if (!StateMachineExecutionStack.IsEmpty())
 	{
-		return nullptr;
+		return GetStateMachine(StateMachineExecutionStack.Last());
 	}
-
-	return GetStateMachine(StateMachineExecutionStack.Last());
+	return nullptr;
 }
 
 const FSceneState* FSceneStateExecutionContext::GetActiveState(const FSceneStateMachine& InStateMachine) const
@@ -421,12 +429,7 @@ const FSceneState* FSceneStateExecutionContext::GetActiveState(const FSceneState
 
 const FSceneState* FSceneStateExecutionContext::GetState(uint16 InAbsoluteIndex) const
 {
-	if (!GeneratedClass)
-	{
-		return nullptr;
-	}
-
-	const TConstArrayView<FSceneState> StatesView = GeneratedClass->GetStates();
+	const TConstArrayView<FSceneState> StatesView = GetSharedTemplateData().States;
 
 	checkfSlow(StatesView.IsValidIndex(InAbsoluteIndex)
 		, TEXT("[Context: %s] Absolute Index %d is out of bounds! (Total State Count: %d)")
@@ -439,12 +442,7 @@ const FSceneState* FSceneStateExecutionContext::GetState(uint16 InAbsoluteIndex)
 
 const FSceneStateEventHandler* FSceneStateExecutionContext::GetEventHandler(uint16 InAbsoluteIndex) const
 {
-	if (!GeneratedClass)
-	{
-		return nullptr;
-	}
-
-	const TConstArrayView<FSceneStateEventHandler> EventHandlers = GeneratedClass->GetEventHandlers();
+	const TConstArrayView<FSceneStateEventHandler> EventHandlers = GetSharedTemplateData().EventHandlers;
 
 	checkfSlow(EventHandlers.IsValidIndex(InAbsoluteIndex)
 		, TEXT("[Context: %s] Absolute Index %d is out of bounds! (Total Event Handlers Count: %d)")
@@ -458,12 +456,12 @@ const FSceneStateEventHandler* FSceneStateExecutionContext::GetEventHandler(uint
 const FSceneState* FSceneStateExecutionContext::GetState(const FSceneStateMachine& InStateMachine, uint16 InRelativeIndex) const
 {
 	const FSceneStateRange StateRange = InStateMachine.GetStateRange();
-	if (StateRange.Count == 0 || !GeneratedClass)
+	if (StateRange.Count == 0)
 	{
 		return nullptr;
 	}
 
-	const TConstArrayView<FSceneState> StatesView = GeneratedClass->GetStates();
+	const TConstArrayView<FSceneState> StatesView = GetSharedTemplateData().States;
 
 	const uint16 AbsoluteIndex = StateRange.Index + InRelativeIndex;
 
@@ -481,12 +479,12 @@ const FSceneState* FSceneStateExecutionContext::GetState(const FSceneStateMachin
 const FSceneStateConduit* FSceneStateExecutionContext::GetConduit(const FSceneStateMachine& InStateMachine, uint16 InRelativeIndex) const
 {
 	const FSceneStateRange ConduitRange = InStateMachine.GetConduitRange();
-	if (ConduitRange.Count == 0 || !GeneratedClass)
+	if (ConduitRange.Count == 0)
 	{
 		return nullptr;
 	}
 
-	const TConstArrayView<FSceneStateConduit> ConduitsView = GeneratedClass->GetConduits();
+	const TConstArrayView<FSceneStateConduit> ConduitsView = GetSharedTemplateData().Conduits;
 
 	const uint16 AbsoluteIndex = ConduitRange.Index + InRelativeIndex;
 
@@ -503,12 +501,8 @@ const FSceneStateConduit* FSceneStateExecutionContext::GetConduit(const FSceneSt
 
 FSceneStateInstance* FSceneStateExecutionContext::FindOrAddStateInstance(const FSceneState& InState) const
 {
-	uint16 StateIndex;
-	if (GetStateIndex(InState, /*out*/StateIndex))
-	{
-		return &StateInstances.FindOrAdd(StateIndex);
-	}
-	return nullptr;
+	const FSceneStateTemplateData& TemplateData = GetSharedTemplateData();
+	return &StateInstances.FindOrAdd(GetStateIndex(TemplateData, InState));
 }
 
 FSceneStateInstance* FSceneStateExecutionContext::FindStateInstance(uint16 InAbsoluteIndex) const
@@ -518,58 +512,37 @@ FSceneStateInstance* FSceneStateExecutionContext::FindStateInstance(uint16 InAbs
 
 FSceneStateInstance* FSceneStateExecutionContext::FindStateInstance(const FSceneState& InState) const
 {
-	uint16 StateIndex;
-	if (GetStateIndex(InState, /*out*/StateIndex))
-	{
-		return FindStateInstance(StateIndex);
-	}
-	return nullptr;
+	const FSceneStateTemplateData& TemplateData = GetSharedTemplateData();
+	return FindStateInstance(GetStateIndex(TemplateData, InState));
 }
 
 void FSceneStateExecutionContext::RemoveStateInstance(const FSceneState& InState) const
 {
-	uint16 StateIndex;
-	if (GetStateIndex(InState, /*out*/StateIndex))
-	{
-		StateInstances.Remove(StateIndex);
-	}
+	const FSceneStateTemplateData& TemplateData = GetSharedTemplateData();
+	StateInstances.Remove(GetStateIndex(TemplateData, InState));
 }
 
 FConstStructView FSceneStateExecutionContext::FindTask(uint16 InAbsoluteIndex) const
 {
-	if (!GeneratedClass)
-	{
-		return {};
-	}
-
-	const FInstancedStructContainer& Tasks = GeneratedClass->GetTasks();
-
-	if (!Tasks.IsValidIndex(InAbsoluteIndex))
+	const FInstancedStructContainer& Tasks = GetSharedTemplateData().Tasks;
+	if (Tasks.IsValidIndex(InAbsoluteIndex))
 	{
-		return {};
+		return Tasks[InAbsoluteIndex];
 	}
-
-	return Tasks[InAbsoluteIndex];
+	return {};
 }
 
 FInstancedStructContainer* FSceneStateExecutionContext::FindOrAddTaskInstanceContainer(const FSceneState& InState) const
 {
-	uint16 StateIndex;
-	if (GetStateIndex(InState, /*out*/StateIndex))
-	{
-		return &TaskInstanceContainers.FindOrAdd(StateIndex);
-	}
-	return nullptr;
+	// Note: task instance containers are allocated per state
+	const FSceneStateTemplateData& TemplateData = GetSharedTemplateData();
+	return &TaskInstanceContainers.FindOrAdd(GetStateIndex(TemplateData, InState));
 }
 
 FInstancedStructContainer* FSceneStateExecutionContext::FindTaskInstanceContainer(const FSceneState& InState) const
 {
-	uint16 StateIndex;
-	if (GetStateIndex(InState, /*out*/StateIndex))
-	{
-		return FindTaskInstanceContainer(StateIndex);
-	}
-	return nullptr;
+	const FSceneStateTemplateData& TemplateData = GetSharedTemplateData();
+	return FindTaskInstanceContainer(GetStateIndex(TemplateData, InState));
 }
 
 FInstancedStructContainer* FSceneStateExecutionContext::FindTaskInstanceContainer(uint16 InAbsoluteIndex) const
@@ -579,13 +552,9 @@ FInstancedStructContainer* FSceneStateExecutionContext::FindTaskInstanceContaine
 
 FStructView FSceneStateExecutionContext::FindTaskInstance(uint16 InAbsoluteIndex) const
 {
-	if (!GeneratedClass)
-	{
-		return {};
-	}
-
-	const FInstancedStructContainer& Tasks = GeneratedClass->GetTasks();
+	const FSceneStateTemplateData& TemplateData = GetSharedTemplateData();
 
+	const FInstancedStructContainer& Tasks = TemplateData.Tasks;
 	if (!Tasks.IsValidIndex(InAbsoluteIndex))
 	{
 		return {};
@@ -601,7 +570,7 @@ FStructView FSceneStateExecutionContext::FindTaskInstance(uint16 InAbsoluteIndex
 		return {};
 	}
 
-	const TConstArrayView<FSceneState> StatesView = GeneratedClass->GetStates();
+	const TConstArrayView<FSceneState> StatesView = TemplateData.States;
 
 	checkf(StatesView.IsValidIndex(StateIndex)
 		, TEXT("[Context: %s] Task Parent State Index %d is out of bounds! (Total State Count: %d)")
@@ -622,21 +591,14 @@ FStructView FSceneStateExecutionContext::FindTaskInstance(uint16 InAbsoluteIndex
 
 void FSceneStateExecutionContext::RemoveTaskInstanceContainer(const FSceneState& InState) const
 {
-	uint16 StateIndex;
-	if (GetStateIndex(InState, /*out*/StateIndex))
-	{
-		TaskInstanceContainers.Remove(StateIndex);
-	}
+	const FSceneStateTemplateData& TemplateData = GetSharedTemplateData();
+	TaskInstanceContainers.Remove(GetStateIndex(TemplateData, InState));
 }
 
 FSceneStateMachineInstance* FSceneStateExecutionContext::FindOrAddStateMachineInstance(const FSceneStateMachine& InStateMachine) const
 {
-	uint16 StateMachineIndex;
-	if (GetStateMachineIndex(InStateMachine, /*out*/StateMachineIndex))
-	{
-		return &StateMachineInstances.FindOrAdd(StateMachineIndex);
-	}
-	return nullptr;
+	const FSceneStateTemplateData& TemplateData = GetSharedTemplateData();
+	return &StateMachineInstances.FindOrAdd(GetStateMachineIndex(TemplateData, InStateMachine));
 }
 
 FSceneStateMachineInstance* FSceneStateExecutionContext::FindStateMachineInstance(uint16 InAbsoluteIndex) const
@@ -646,31 +608,20 @@ FSceneStateMachineInstance* FSceneStateExecutionContext::FindStateMachineInstanc
 
 FSceneStateMachineInstance* FSceneStateExecutionContext::FindStateMachineInstance(const FSceneStateMachine& InStateMachine) const
 {
-	uint16 StateMachineIndex;
-	if (GetStateMachineIndex(InStateMachine, /*out*/StateMachineIndex))
-	{
-		return FindStateMachineInstance(StateMachineIndex);
-	}
-	return nullptr;
+	const FSceneStateTemplateData& TemplateData = GetSharedTemplateData();
+	return FindStateMachineInstance(GetStateMachineIndex(TemplateData, InStateMachine));
 }
 
 void FSceneStateExecutionContext::RemoveStateMachineInstance(const FSceneStateMachine& InStateMachine) const
 {
-	uint16 StateMachineIndex;
-	if (GetStateMachineIndex(InStateMachine, /*out*/StateMachineIndex))
-	{
-		StateMachineInstances.Remove(StateMachineIndex);
-	}
+	const FSceneStateTemplateData& TemplateData = GetSharedTemplateData();
+	StateMachineInstances.Remove(GetStateMachineIndex(TemplateData, InStateMachine));
 }
 
 FSceneStateTransitionInstance* FSceneStateExecutionContext::FindOrAddTransitionInstance(const FSceneStateTransition& InTransition) const
 {
-	uint16 TransitionIndex;
-	if (GetTransitionIndex(InTransition, /*out*/TransitionIndex))
-	{
-		return &TransitionInstances.FindOrAdd(TransitionIndex);
-	}
-	return nullptr;
+	const FSceneStateTemplateData& TemplateData = GetSharedTemplateData();
+	return &TransitionInstances.FindOrAdd(GetTransitionIndex(TemplateData, InTransition));
 }
 
 FSceneStateTransitionInstance* FSceneStateExecutionContext::FindTransitionInstance(uint16 InAbsoluteIndex) const
@@ -680,40 +631,26 @@ FSceneStateTransitionInstance* FSceneStateExecutionContext::FindTransitionInstan
 
 FSceneStateTransitionInstance* FSceneStateExecutionContext::FindTransitionInstance(const FSceneStateTransition& InTransition) const
 {
-	uint16 TransitionIndex;
-	if (GetTransitionIndex(InTransition, /*out*/TransitionIndex))
-	{
-		return FindTransitionInstance(TransitionIndex);
-	}
-	return nullptr;
+	const FSceneStateTemplateData& TemplateData = GetSharedTemplateData();
+	return FindTransitionInstance(GetTransitionIndex(TemplateData, InTransition));
 }
 
 void FSceneStateExecutionContext::RemoveTransitionInstance(const FSceneStateTransition& InTransition) const
 {
-	uint16 TransitionIndex;
-	if (GetTransitionIndex(InTransition, /*out*/TransitionIndex))
-	{
-		TransitionInstances.Remove(TransitionIndex);
-	}
+	const FSceneStateTemplateData& TemplateData = GetSharedTemplateData();
+	TransitionInstances.Remove(GetTransitionIndex(TemplateData, InTransition));
 }
 
 FSceneStateConduitInstance* FSceneStateExecutionContext::FindOrAddConduitInstance(const FSceneStateConduit& InConduit) const
 {
-	uint16 ConduitIndex;
-	if (GetConduitIndex(InConduit, /*out*/ConduitIndex))
-	{
-		return &ConduitInstances.FindOrAdd(ConduitIndex);
-	}
-	return nullptr;
+	const FSceneStateTemplateData& TemplateData = GetSharedTemplateData();
+	return &ConduitInstances.FindOrAdd(GetConduitIndex(TemplateData, InConduit));
 }
 
 void FSceneStateExecutionContext::RemoveConduitInstance(const FSceneStateConduit& InConduit) const
 {
-	uint16 ConduitIndex;
-	if (GetConduitIndex(InConduit, /*out*/ConduitIndex))
-	{
-		ConduitInstances.Remove(ConduitIndex);
-	}
+	const FSceneStateTemplateData& TemplateData = GetSharedTemplateData();
+	ConduitInstances.Remove(GetConduitIndex(TemplateData, InConduit));
 }
 
 bool FSceneStateExecutionContext::SetupFunctionInstances(FPropertyBindingIndex16 InBindingsBatch) const
@@ -724,16 +661,9 @@ bool FSceneStateExecutionContext::SetupFunctionInstances(FPropertyBindingIndex16
 		return false;
 	}
 
-	const FSceneStateBindingCollection* const BindingCollection = GetBindingCollection();
-	if (!BindingCollection)
-	{
-		UE_LOG(LogSceneState, Error, TEXT("[%s] Failed to setup function instances for batch %d. Binding Collection is null")
-			, *GetExecutionContextName()
-			, InBindingsBatch.Get());
-		return false;
-	}
+	const FSceneStateTemplateData& TemplateData = GetSharedTemplateData();
 
-	const FPropertyBindingCopyInfoBatch& CopyInfoBatch = BindingCollection->GetBatch(InBindingsBatch);
+	const FPropertyBindingCopyInfoBatch& CopyInfoBatch = TemplateData.BindingCollection.GetBatch(InBindingsBatch);
 
 	if (!UE::SceneState::IsValidRange(CopyInfoBatch.PropertyFunctionsBegin, CopyInfoBatch.PropertyFunctionsEnd))
 	{
@@ -744,14 +674,14 @@ bool FSceneStateExecutionContext::SetupFunctionInstances(FPropertyBindingIndex16
 	FSceneStateFunctionInstanceBatch& FunctionInstanceBatch = FunctionInstanceBatches.Add(InBindingsBatch);
 	FunctionInstanceBatch.FunctionRange = FSceneStateRange::MakeBeginEndRange(CopyInfoBatch.PropertyFunctionsBegin.Get(), CopyInfoBatch.PropertyFunctionsEnd.Get());
 
-	// Initialize the function instances to the templates in the generated class
+	// Initialize the function instances to the templates in the template data
 	{
 		TArray<FConstStructView> TemplateFunctionInstances;
 		TemplateFunctionInstances.Reserve(FunctionInstanceBatch.FunctionRange.Count);
 
 		for (uint16 FunctionIndex = FunctionInstanceBatch.FunctionRange.Index; FunctionIndex <= FunctionInstanceBatch.FunctionRange.GetLastIndex(); ++FunctionIndex)
 		{
-			TemplateFunctionInstances.Add(GeneratedClass->GetFunctionInstances()[FunctionIndex]);
+			TemplateFunctionInstances.Add(TemplateData.FunctionInstances[FunctionIndex]);
 		}
 
 		FunctionInstanceBatch.FunctionInstances.Append(TemplateFunctionInstances);
@@ -767,7 +697,7 @@ bool FSceneStateExecutionContext::SetupFunctionInstances(FPropertyBindingIndex16
 	// Set up these functions as they could be bound to other functions
 	for (uint16 FunctionIndex = FunctionInstanceBatch.FunctionRange.Index; FunctionIndex <= FunctionInstanceBatch.FunctionRange.GetLastIndex(); ++FunctionIndex)
 	{
-		if (const FSceneStateFunction* Function = GeneratedClass->GetFunctions()[FunctionIndex].GetPtr<const FSceneStateFunction>())
+		if (const FSceneStateFunction* Function = TemplateData.Functions[FunctionIndex].GetPtr<const FSceneStateFunction>())
 		{
 			Function->Setup(*this);
 		}
@@ -793,10 +723,12 @@ void FSceneStateExecutionContext::RemoveFunctionInstances(FPropertyBindingIndex1
 		return;
 	}
 
+	const FSceneStateTemplateData& TemplateData = GetSharedTemplateData();
+
 	// Clean up the functions used resources (e.g. other function instance data)
 	for (uint16 FunctionIndex = FunctionRange.Index; FunctionIndex <= FunctionRange.GetLastIndex(); ++FunctionIndex)
 	{
-		if (const FSceneStateFunction* Function = GeneratedClass->GetFunctions()[FunctionIndex].GetPtr<const FSceneStateFunction>())
+		if (const FSceneStateFunction* Function = TemplateData.Functions[FunctionIndex].GetPtr<const FSceneStateFunction>())
 		{
 			Function->Cleanup(*this);
 		}
@@ -812,19 +744,12 @@ void FSceneStateExecutionContext::RemoveFunctionInstances(FPropertyBindingIndex1
 
 FConstStructView FSceneStateExecutionContext::FindFunction(uint16 InAbsoluteIndex) const
 {
-	if (!GeneratedClass)
+	const FInstancedStructContainer& Functions = GetSharedTemplateData().Functions;
+	if (Functions.IsValidIndex(InAbsoluteIndex))
 	{
-		return {};
-	}
-
-	const FInstancedStructContainer& Functions = GeneratedClass->GetFunctions();
-
-	if (!Functions.IsValidIndex(InAbsoluteIndex))
-	{
-		return {};
+		return Functions[InAbsoluteIndex];
 	}
-
-	return Functions[InAbsoluteIndex];
+	return {};
 }
 
 FStructView FSceneStateExecutionContext::FindFunctionInstance(uint16 InAbsoluteIndex) const
@@ -841,54 +766,34 @@ const FSceneStateExecutionContext* FSceneStateExecutionContext::FindParentContex
 	return nullptr;
 }
 
-bool FSceneStateExecutionContext::GetStateIndex(const FSceneState& InState, uint16& OutIndex) const
+uint16 FSceneStateExecutionContext::GetStateIndex(const FSceneStateTemplateData& InTemplateData, const FSceneState& InState) const
 {
-	if (GeneratedClass)
-	{
-		TConstArrayView<FSceneState> StatesView = GeneratedClass->GetStates();
-		const int32 Index = static_cast<int32>(&InState - StatesView.GetData());
-		checkSlow(StatesView.IsValidIndex(Index) && Index >= 0 && Index <= TNumericLimits<uint16>::Max());
-		OutIndex = static_cast<uint16>(Index);
-		return true;
-	}
-	return false;
+	TConstArrayView<FSceneState> StatesView = GetSharedTemplateData().States;
+	const int32 Index = static_cast<int32>(&InState - StatesView.GetData());
+	checkSlow(StatesView.IsValidIndex(Index) && Index >= 0 && Index <= TNumericLimits<uint16>::Max());
+	return static_cast<uint16>(Index);
 }
 
-bool FSceneStateExecutionContext::GetStateMachineIndex(const FSceneStateMachine& InStateMachine, uint16& OutIndex) const
+uint16 FSceneStateExecutionContext::GetStateMachineIndex(const FSceneStateTemplateData& InTemplateData, const FSceneStateMachine& InStateMachine) const
 {
-	if (GeneratedClass)
-	{
-		TConstArrayView<FSceneStateMachine> StateMachinesView = GeneratedClass->GetStateMachines();
-		const int32 Index = static_cast<int32>(&InStateMachine - StateMachinesView.GetData());
-		checkSlow(StateMachinesView.IsValidIndex(Index) && Index >= 0 && Index <= TNumericLimits<uint16>::Max());
-		OutIndex = static_cast<uint16>(Index);
-		return true;
-	}
-	return false;
+	TConstArrayView<FSceneStateMachine> StateMachinesView = InTemplateData.StateMachines;
+	const int32 Index = static_cast<int32>(&InStateMachine - StateMachinesView.GetData());
+	checkSlow(StateMachinesView.IsValidIndex(Index) && Index >= 0 && Index <= TNumericLimits<uint16>::Max());
+	return static_cast<uint16>(Index);
 }
 
-bool FSceneStateExecutionContext::GetTransitionIndex(const FSceneStateTransition& InTransition, uint16& OutIndex) const
+uint16 FSceneStateExecutionContext::GetTransitionIndex(const FSceneStateTemplateData& InTemplateData, const FSceneStateTransition& InTransition) const
 {
-	if (GeneratedClass)
-	{
-		TConstArrayView<FSceneStateTransition> TransitionsView = GeneratedClass->GetTransitions();
-		const int32 Index = static_cast<int32>(&InTransition - TransitionsView.GetData());
-		checkSlow(TransitionsView.IsValidIndex(Index) && Index >= 0 && Index <= TNumericLimits<uint16>::Max());
-		OutIndex = static_cast<uint16>(Index);
-		return true;
-	}
-	return false;
+	TConstArrayView<FSceneStateTransition> TransitionsView = InTemplateData.Transitions;
+	const int32 Index = static_cast<int32>(&InTransition - TransitionsView.GetData());
+	checkSlow(TransitionsView.IsValidIndex(Index) && Index >= 0 && Index <= TNumericLimits<uint16>::Max());
+	return static_cast<uint16>(Index);
 }
 
-bool FSceneStateExecutionContext::GetConduitIndex(const FSceneStateConduit& InConduit, uint16& OutIndex) const
+uint16 FSceneStateExecutionContext::GetConduitIndex(const FSceneStateTemplateData& InTemplateData, const FSceneStateConduit& InConduit) const
 {
-	if (GeneratedClass)
-	{
-		TConstArrayView<FSceneStateConduit> ConduitsView = GeneratedClass->GetConduits();
-		const int32 Index = static_cast<int32>(&InConduit - ConduitsView.GetData());
-		checkSlow(ConduitsView.IsValidIndex(Index) && Index >= 0 && Index <= TNumericLimits<uint16>::Max());
-		OutIndex = static_cast<uint16>(Index);
-		return true;
-	}
-	return false;
+	TConstArrayView<FSceneStateConduit> ConduitsView = InTemplateData.Conduits;
+	const int32 Index = static_cast<int32>(&InConduit - ConduitsView.GetData());
+	checkSlow(ConduitsView.IsValidIndex(Index) && Index >= 0 && Index <= TNumericLimits<uint16>::Max());
+	return static_cast<uint16>(Index);
 }
diff --git a/Engine/Plugins/VirtualProduction/SceneState/Source/SceneState/Private/SceneStateExecutionScope.cpp b/Engine/Plugins/VirtualProduction/SceneState/Source/SceneState/Private/SceneStateExecutionScope.cpp
--- a/Engine/Plugins/VirtualProduction/SceneState/Source/SceneState/Private/SceneStateExecutionScope.cpp
+++ b/Engine/Plugins/VirtualProduction/SceneState/Source/SceneState/Private/SceneStateExecutionScope.cpp
@@ -9,7 +9,8 @@ namespace UE::SceneState
 FExecutionScope::FExecutionScope(const FSceneStateExecutionContext& InContext, const FSceneStateMachine& InStateMachine)
 	: ExecutionContext(InContext)
 {
-	ensure(ExecutionContext.GetStateMachineIndex(InStateMachine, StateMachineIndex));
+	const FSceneStateTemplateData& TemplateData = InContext.GetSharedTemplateData();
+	StateMachineIndex = ExecutionContext.GetStateMachineIndex(TemplateData, InStateMachine);
 	ExecutionContext.StateMachineExecutionStack.Push(StateMachineIndex);
 }
 
diff --git a/Engine/Plugins/VirtualProduction/SceneState/Source/SceneState/Private/SceneStateGeneratedClass.cpp b/Engine/Plugins/VirtualProduction/SceneState/Source/SceneState/Private/SceneStateGeneratedClass.cpp
--- a/Engine/Plugins/VirtualProduction/SceneState/Source/SceneState/Private/SceneStateGeneratedClass.cpp
+++ b/Engine/Plugins/VirtualProduction/SceneState/Source/SceneState/Private/SceneStateGeneratedClass.cpp
@@ -1,19 +1,7 @@
 // Copyright Epic Games, Inc. All Rights Reserved.
 
 #include "SceneStateGeneratedClass.h"
-#include "Misc/EnumerateRange.h"
-#include "SceneStateBindingUtils.h"
-#include "SceneStateEventSchema.h"
-#include "SceneStateExecutionContext.h"
-#include "SceneStateExecutionContextRegistry.h"
-#include "SceneStateLog.h"
-#include "SceneStateMachine.h"
-#include "SceneStateObject.h"
-#include "SceneStateUtils.h"
-#include "StructUtils/UserDefinedStruct.h"
-#include "Tasks/SceneStateTask.h"
-#include "Tasks/SceneStateTaskBindingExtension.h"
-#include "Tasks/SceneStateTaskInstance.h"
+#include "SceneStateTemplateData.h"
 
 #if WITH_EDITOR
 #include "StructUtils/UserDefinedStruct.h"
@@ -22,6 +10,8 @@
 
 USceneStateGeneratedClass::USceneStateGeneratedClass()
 {
+	SharedTemplateData.InitializeAs<FSceneStateTemplateData>(TNotNull<const UClass*>(this));
+
 #if WITH_EDITOR
 	if (!HasAnyFlags(RF_ClassDefaultObject))
 	{
@@ -31,58 +21,24 @@ USceneStateGeneratedClass::USceneStateGeneratedClass()
 #endif
 }
 
-const FSceneState* USceneStateGeneratedClass::GetRootState() const
-{
-	if (States.IsValidIndex(RootStateIndex))
-	{
-		return &States[RootStateIndex];	
-	}
-	return nullptr;
-}
-
-const FSceneStateMachine* USceneStateGeneratedClass::FindStateMachine(const FGuid& InStateMachineId) const
+FConstSharedStruct USceneStateGeneratedClass::GetSharedTemplateData() const
 {
-	const uint16* StateMachineIndex = StateMachineIdToIndex.Find(InStateMachineId);
-	if (StateMachineIndex && StateMachines.IsValidIndex(*StateMachineIndex))
-	{
-		return &StateMachines[*StateMachineIndex];
-	}
-	return nullptr;
-}
-
-FInstancedPropertyBag USceneStateGeneratedClass::GetTransitionParameters(uint16 InAbsoluteIndex) const
-{
-	return TransitionParameters.FindRef(InAbsoluteIndex);
+	return SharedTemplateData;
 }
 
 void USceneStateGeneratedClass::Link(FArchive& Ar, bool bInRelinkExistingProperties)
 {
 	Super::Link(Ar, bInRelinkExistingProperties);
-
 	if (IsFullClass())
 	{
-		for (TEnumerateRef<FSceneStateTransition> Transition : EnumerateRange(Transitions))
-		{
-			if (TransitionLinks.IsValidIndex(Transition.GetIndex()))
-			{
-				Transition->Link(TransitionLinks[Transition.GetIndex()], this);
-			}
-		}
-
-		for (TEnumerateRef<FSceneStateConduit> Conduit : EnumerateRange(Conduits))
-		{
-			if (ConduitLinks.IsValidIndex(Conduit.GetIndex()))
-			{
-				Conduit->Link(ConduitLinks[Conduit.GetIndex()], this);
-			}
-		}
+		GetTemplateData().Link();
 	}
 }
 
 void USceneStateGeneratedClass::PurgeClass(bool bInRecompilingOnLoad)
 {
 	Super::PurgeClass(bInRecompilingOnLoad);
-	Reset();
+	GetTemplateData().Reset();
 }
 
 void USceneStateGeneratedClass::PostLoad()
@@ -94,171 +50,23 @@ void USceneStateGeneratedClass::PostLoad()
 void USceneStateGeneratedClass::BeginDestroy()
 {
 	Super::BeginDestroy();
-	Reset();
+	GetTemplateData().Reset();
+	SharedTemplateData.Reset();
 }
 
-#if WITH_EDITOR
-const FSceneState* USceneStateGeneratedClass::FindStateFromNode(FObjectKey InStateNode) const
+FSceneStateTemplateData& USceneStateGeneratedClass::GetTemplateData()
 {
-	const uint16* StateIndex = StateNodeToIndex.Find(InStateNode);
-	if (StateIndex && States.IsValidIndex(*StateIndex))
-	{
-		return &States[*StateIndex];
-	}
-	return nullptr;
-}
-
-const FSceneStateTask* USceneStateGeneratedClass::FindTaskFromNode(FObjectKey InTaskNode) const
-{
-	const uint16* TaskIndex = TaskNodeToIndex.Find(InTaskNode);
-	if (TaskIndex && Tasks.IsValidIndex(*TaskIndex))
-	{
-		return Tasks[*TaskIndex].GetPtr<const FSceneStateTask>();
-	}
-	return nullptr;
-}
-
-void USceneStateGeneratedClass::ForEachStateInstance(USceneStateObject& InRootState, FObjectKey InStateNode, TFunctionRef<void(const FSceneStateInstance&)> InFunctor) const
-{
-	const uint16* StateIndex = StateNodeToIndex.Find(InStateNode);
-	if (!StateIndex)
-	{
-		return;
-	}
-
-	const TSharedRef<UE::SceneState::FExecutionContextRegistry> ContextRegistry = InRootState.GetContextRegistry();
-
-	ContextRegistry->ForEachExecutionContext(
-		[StateIndex, &InFunctor](const FSceneStateExecutionContext& InExecutionContext)
-		{
-			if (const FSceneStateInstance* StateInstance = InExecutionContext.FindStateInstance(*StateIndex))
-			{
-				InFunctor(*StateInstance);
-			}
-		});
-}
-
-void USceneStateGeneratedClass::ForEachTaskInstance(USceneStateObject& InRootState, FObjectKey InTaskNode, TFunctionRef<void(const FSceneStateTaskInstance&)> InFunctor) const
-{
-	const uint16* TaskIndex = TaskNodeToIndex.Find(InTaskNode);
-	if (!TaskIndex)
-	{
-		return;
-	}
-
-	const TSharedRef<UE::SceneState::FExecutionContextRegistry> ContextRegistry = InRootState.GetContextRegistry();
-
-	ContextRegistry->ForEachExecutionContext(
-		[TaskIndex, &InFunctor](const FSceneStateExecutionContext& InExecutionContext)
-		{
-			if (const FSceneStateTaskInstance* TaskInstance = InExecutionContext.FindTaskInstance(*TaskIndex).GetPtr<const FSceneStateTaskInstance>())
-			{
-				InFunctor(*TaskInstance);
-			}
-		});
-}
-#endif
-
-const UStruct* USceneStateGeneratedClass::FindDataStruct(const FSceneStateBindingDataHandle& InDataHandle)
-{
-	const uint16 DataIndex = InDataHandle.GetDataIndex();
-	const uint16 DataSubIndex = InDataHandle.GetDataSubIndex();
-
-	if (InDataHandle.IsExternalDataType())
-	{
-		// Todo: unsupported external sources
-		return nullptr;
-	}
-
-	const ESceneStateDataType DataType = static_cast<ESceneStateDataType>(InDataHandle.GetDataType());
-
-	switch (DataType)
-	{
-	case ESceneStateDataType::Root:
-		return this;
-
-	case ESceneStateDataType::Task:
-		return TaskInstances[DataIndex].GetScriptStruct();
-
-	case ESceneStateDataType::EventHandler:
-		return EventHandlers[DataIndex].GetEventSchemaHandle().GetEventStruct();
-
-	case ESceneStateDataType::TaskExtension:
-		if (const FSceneStateTask* Task = Tasks[DataIndex].GetPtr<const FSceneStateTask>())
-		{
-			if (const FSceneStateTaskBindingExtension* BindingExtension = Task->GetBindingExtension())
-			{
-				FStructView DataView;
-				if (BindingExtension->FindDataByIndex(TaskInstances[DataIndex], DataSubIndex, DataView))
-				{
-					return DataView.GetScriptStruct();
-				}
-			}
-		}
-		break;
-
-	case ESceneStateDataType::Transition:
-		return TransitionParameters[DataIndex].GetPropertyBagStruct();
-
-	case ESceneStateDataType::StateMachine:
-		return StateMachines[DataIndex].GetParametersStruct();
-
-	case ESceneStateDataType::Function:
-		return FunctionInstances[DataIndex].GetScriptStruct();
-	}
-
-	return nullptr;
+	return SharedTemplateData.Get<FSceneStateTemplateData>();
 }
 
 void USceneStateGeneratedClass::ResolveBindings()
 {
-	if (!IsFullClass())
-	{
-		return;	
-	}
-
-	UE::SceneState::PatchBindingCollection({ .BindingCollection = BindingCollection, .FindDataStructFunctor =
-		[this](const FSceneStateBindingDataHandle& InDataHandle)->const UStruct*
-		{
-			return FindDataStruct(InDataHandle);
-		}});
-
-	// Resolves property paths used by bindings a store property pointers
-	if (!BindingCollection.ResolvePaths())
+	if (IsFullClass())
 	{
-		UE_LOG(LogSceneState, Warning, TEXT("%s failed to resolve bindings. Try compiling Scene State Blueprint again"), *GetName());
+		GetTemplateData().ResolveBindings();
 	}
 }
 
-void USceneStateGeneratedClass::Reset()
-{
-	Tasks.Reset();
-	TaskPrerequisites.Reset();
-	TaskInstances.Reset();
-	Functions.Reset();
-	FunctionInstances.Reset();
-	States.Reset();
-	Conduits.Reset();
-	ConduitLinks.Reset();
-	EventHandlers.Reset();
-	Transitions.Reset();
-	TransitionLinks.Reset();
-	TransitionParameters.Reset();
-	StateMachines.Reset();
-	BindingCollection.Reset();
-	StateMachineIdToIndex.Reset();
-
-#if WITH_EDITOR
-	TaskMetadata.Reset();
-	FunctionMetadata.Reset();
-	StateMetadata.Reset();
-	TransitionMetadata.Reset();
-	StateNodeToIndex.Reset();
-	StateMachineGraphToIndex.Reset();
-	TaskNodeToIndex.Reset();
-#endif
-}
-
 bool USceneStateGeneratedClass::IsFullClass() const
 {
 #if WITH_EDITOR
@@ -273,38 +81,17 @@ bool USceneStateGeneratedClass::IsFullClass() const
 #if WITH_EDITOR
 void USceneStateGeneratedClass::OnObjectsReinstanced(const TMap<UObject*, UObject*>& InReplacementMap)
 {
-	TSet<const UStruct*> Structs;
-	Structs.Reserve(InReplacementMap.Num());
-
-	bool bRequiresResolve = false;
-
-	for (const TPair<UObject*, UObject*>& Pair : InReplacementMap)
-	{
-		if (const UObject* Replacement = Pair.Value)
-		{
-			if (Replacement->IsIn(this))
-			{
-				bRequiresResolve = true;
-				break;
-			}
-			Structs.Add(Replacement->GetClass());
-		}
-	}
-
-	if (bRequiresResolve || BindingCollection.ContainsAnyStruct(Structs))
+	if (IsFullClass())
 	{
-		ResolveBindings();
+		GetTemplateData().OnObjectsReinstanced(InReplacementMap);
 	}
 }
 
 void USceneStateGeneratedClass::OnStructsReinstanced(const UUserDefinedStruct& InStruct)
 {
-	TSet<const UStruct*> Structs;
-	Structs.Add(&InStruct);
-
-	if (BindingCollection.ContainsAnyStruct(Structs))
+	if (IsFullClass())
 	{
-		ResolveBindings();
+		GetTemplateData().OnStructsReinstanced(InStruct);
 	}
 }
 #endif
diff --git a/Engine/Plugins/VirtualProduction/SceneState/Source/SceneState/Private/SceneStateObject.cpp b/Engine/Plugins/VirtualProduction/SceneState/Source/SceneState/Private/SceneStateObject.cpp
--- a/Engine/Plugins/VirtualProduction/SceneState/Source/SceneState/Private/SceneStateObject.cpp
+++ b/Engine/Plugins/VirtualProduction/SceneState/Source/SceneState/Private/SceneStateObject.cpp
@@ -7,6 +7,7 @@
 #include "SceneStateExecutionContextRegistry.h"
 #include "SceneStateGeneratedClass.h"
 #include "SceneStatePlayer.h"
+#include "SceneStateTemplateData.h"
 
 USceneStateObject::USceneStateObject()
 {
@@ -24,6 +25,11 @@ FString USceneStateObject::GetContextName() const
 	return FString();
 }
 
+FConstSharedStruct USceneStateObject::GetSharedTemplateData() const
+{
+	return SharedTemplateData;
+}
+
 UObject* USceneStateObject::GetContextObject() const
 {
 	if (USceneStatePlayer* Player = Cast<USceneStatePlayer>(GetOuter()))
@@ -35,20 +41,20 @@ UObject* USceneStateObject::GetContextObject() const
 
 bool USceneStateObject::IsActive() const
 {
-	if (GeneratedClass)
+	if (const FSceneState* RootState = GetRootState())
 	{
-		if (const FSceneState* RootState = GeneratedClass->GetRootState())
-		{
-			const FSceneStateInstance* RootStateInstance = RootExecutionContext.FindStateInstance(*RootState);
-			return RootStateInstance && RootStateInstance->Status == UE::SceneState::EExecutionStatus::Running;
-		}
+		const FSceneStateInstance* RootStateInstance = RootExecutionContext.FindStateInstance(*RootState);
+		return RootStateInstance && RootStateInstance->Status == UE::SceneState::EExecutionStatus::Running;
 	}
 	return false;
 }
 
 void USceneStateObject::Setup()
 {
-	GeneratedClass = Cast<USceneStateGeneratedClass>(GetClass());
+	if (const USceneStateGeneratedClass* GeneratedClass = Cast<USceneStateGeneratedClass>(GetClass()))
+	{
+		SharedTemplateData = GeneratedClass->GetSharedTemplateData();
+	}
 	RootExecutionContext.Setup(this);
 }
 
@@ -59,24 +65,16 @@ void USceneStateObject::Enter()
 		EventStream->Register();
 	}
 
-	if (GeneratedClass)
+	if (const FSceneState* RootState = GetRootState())
 	{
-		if (const FSceneState* RootState = GeneratedClass->GetRootState())
-		{
-			ReceiveEnter();
-			RootState->Enter(RootExecutionContext);
-		}
+		ReceiveEnter();
+		RootState->Enter(RootExecutionContext);
 	}
 }
 
 void USceneStateObject::Tick(float InDeltaSeconds)
 {
-	if (!GeneratedClass)
-	{
-		return;
-	}
-
-	if (const FSceneState* RootState = GeneratedClass->GetRootState())
+	if (const FSceneState* RootState = GetRootState())
 	{
 		ReceiveTick(InDeltaSeconds);
 		RootState->Tick(RootExecutionContext, InDeltaSeconds);
@@ -85,13 +83,10 @@ void USceneStateObject::Tick(float InDeltaSeconds)
 
 void USceneStateObject::Exit()
 {
-	if (GeneratedClass)
+	if (const FSceneState* RootState = GetRootState())
 	{
-		if (const FSceneState* RootState = GeneratedClass->GetRootState())
-		{
-			ReceiveExit();
-			RootState->Exit(RootExecutionContext);
-		}
+		ReceiveExit();
+		RootState->Exit(RootExecutionContext);
 	}
 
 	if (EventStream)
@@ -99,6 +94,7 @@ void USceneStateObject::Exit()
 		EventStream->Unregister();
 	}
 
+	SharedTemplateData.Reset();
 	RootExecutionContext.Reset();
 }
 
@@ -119,5 +115,15 @@ UWorld* USceneStateObject::GetWorld() const
 void USceneStateObject::BeginDestroy()
 {
 	Super::BeginDestroy();
+	SharedTemplateData.Reset();
     RootExecutionContext.Reset();
 }
+
+const FSceneState* USceneStateObject::GetRootState() const
+{
+	if (const FSceneStateTemplateData* TemplateData = GetSharedTemplateData().GetPtr<const FSceneStateTemplateData>())
+	{
+		return TemplateData->GetRootState();
+	}
+	return nullptr;
+}
diff --git a/Engine/Plugins/VirtualProduction/SceneState/Source/SceneState/Private/SceneStatePlayer.cpp b/Engine/Plugins/VirtualProduction/SceneState/Source/SceneState/Private/SceneStatePlayer.cpp
--- a/Engine/Plugins/VirtualProduction/SceneState/Source/SceneState/Private/SceneStatePlayer.cpp
+++ b/Engine/Plugins/VirtualProduction/SceneState/Source/SceneState/Private/SceneStatePlayer.cpp
@@ -2,7 +2,6 @@
 
 #include "SceneStatePlayer.h"
 #include "SceneState.h"
-#include "SceneStateGeneratedClass.h"
 #include "SceneStateMachine.h"
 #include "SceneStateObject.h"
 #include "SceneStateUtils.h"
diff --git a/Engine/Plugins/VirtualProduction/SceneState/Source/SceneState/Private/SceneStatePropertyUtils.cpp b/Engine/Plugins/VirtualProduction/SceneState/Source/SceneState/Private/SceneStatePropertyUtils.cpp
--- a/Engine/Plugins/VirtualProduction/SceneState/Source/SceneState/Private/SceneStatePropertyUtils.cpp
+++ b/Engine/Plugins/VirtualProduction/SceneState/Source/SceneState/Private/SceneStatePropertyUtils.cpp
@@ -14,13 +14,9 @@ namespace Private
 /** Resolves a property reference without recursing/iterating further if the source of that reference is another property reference */
 bool ResolvePropertyInternal(const FSceneStateExecutionContext& InContext, const FSceneStatePropertyReference& InPropertyReference, FResolvePropertyResult& OutResult)
 {
-	const FSceneStateBindingCollection* BindingCollection = InContext.GetBindingCollection();
-	if (!BindingCollection)
-	{
-		return false;
-	}
+	const FSceneStateBindingCollection& BindingCollection = InContext.GetBindingCollection();
 
-	const FSceneStateBindingResolvedReference* ResolvedReference = BindingCollection->FindResolvedReference(InPropertyReference);
+	const FSceneStateBindingResolvedReference* ResolvedReference = BindingCollection.FindResolvedReference(InPropertyReference);
 	if (!ResolvedReference || !ResolvedReference->SourceLeafProperty)
 	{
 		return false;
@@ -32,7 +28,7 @@ bool ResolvePropertyInternal(const FSceneStateExecutionContext& InContext, const
 		return false;
 	}
 
-	uint8* const ValuePtr = BindingCollection->ResolveProperty(*ResolvedReference, SourceDataView);
+	uint8* const ValuePtr = BindingCollection.ResolveProperty(*ResolvedReference, SourceDataView);
 	if (!ValuePtr)
 	{
 		return false;
diff --git a/Engine/Plugins/VirtualProduction/SceneState/Source/SceneState/Private/SceneStateTemplateData.cpp b/Engine/Plugins/VirtualProduction/SceneState/Source/SceneState/Private/SceneStateTemplateData.cpp
new file mode 100644
--- /dev/null
+++ b/Engine/Plugins/VirtualProduction/SceneState/Source/SceneState/Private/SceneStateTemplateData.cpp
@@ -0,0 +1,250 @@
+// Copyright Epic Games, Inc. All Rights Reserved.
+
+#include "SceneStateTemplateData.h"
+#include "Misc/EnumerateRange.h"
+#include "SceneStateBindingUtils.h"
+#include "SceneStateExecutionContext.h"
+#include "SceneStateExecutionContextRegistry.h"
+#include "SceneStateLog.h"
+#include "StructUtils/UserDefinedStruct.h"
+#include "Tasks/SceneStateTask.h"
+#include "Tasks/SceneStateTaskBindingExtension.h"
+#include "Tasks/SceneStateTaskInstance.h"
+
+FSceneStateTemplateData::FSceneStateTemplateData(TNotNull<const UClass*> InOwnerClass)
+	: OwnerClass(InOwnerClass)
+{
+}
+
+const FSceneState* FSceneStateTemplateData::GetRootState() const
+{
+	if (States.IsValidIndex(RootStateIndex))
+	{
+		return &States[RootStateIndex];	
+	}
+	return nullptr;
+}
+
+const FSceneStateMachine* FSceneStateTemplateData::FindStateMachine(const FGuid& InStateMachineId) const
+{
+	const uint16* StateMachineIndex = StateMachineIdToIndex.Find(InStateMachineId);
+	if (StateMachineIndex && StateMachines.IsValidIndex(*StateMachineIndex))
+	{
+		return &StateMachines[*StateMachineIndex];
+	}
+	return nullptr;
+}
+
+const UStruct* FSceneStateTemplateData::FindDataStruct(const FSceneStateBindingDataHandle& InDataHandle)
+{
+	const uint16 DataIndex = InDataHandle.GetDataIndex();
+	const uint16 DataSubIndex = InDataHandle.GetDataSubIndex();
+
+	if (InDataHandle.IsExternalDataType())
+	{
+		// Todo: unsupported external sources
+		return nullptr;
+	}
+
+	const ESceneStateDataType DataType = static_cast<ESceneStateDataType>(InDataHandle.GetDataType());
+
+	switch (DataType)
+	{
+	case ESceneStateDataType::Root:
+		return OwnerClass;
+
+	case ESceneStateDataType::Task:
+		return TaskInstances[DataIndex].GetScriptStruct();
+
+	case ESceneStateDataType::EventHandler:
+		return EventHandlers[DataIndex].GetEventSchemaHandle().GetEventStruct();
+
+	case ESceneStateDataType::TaskExtension:
+		if (const FSceneStateTask* Task = Tasks[DataIndex].GetPtr<const FSceneStateTask>())
+		{
+			if (const FSceneStateTaskBindingExtension* BindingExtension = Task->GetBindingExtension())
+			{
+				FStructView DataView;
+				if (BindingExtension->FindDataByIndex(TaskInstances[DataIndex], DataSubIndex, DataView))
+				{
+					return DataView.GetScriptStruct();
+				}
+			}
+		}
+		break;
+
+	case ESceneStateDataType::Transition:
+		return TransitionParameters[DataIndex].GetPropertyBagStruct();
+
+	case ESceneStateDataType::StateMachine:
+		return StateMachines[DataIndex].GetParametersStruct();
+
+	case ESceneStateDataType::Function:
+		return FunctionInstances[DataIndex].GetScriptStruct();
+	}
+
+	return nullptr;
+}
+
+void FSceneStateTemplateData::ResolveBindings()
+{
+	UE::SceneState::PatchBindingCollection(
+		{
+			.BindingCollection = BindingCollection,
+			.FindDataStructFunctor = [this](const FSceneStateBindingDataHandle& InDataHandle)->const UStruct*
+			{
+				return FindDataStruct(InDataHandle);
+			}
+		});
+
+	// Resolves property paths used by bindings a store property pointers
+	if (!BindingCollection.ResolvePaths())
+	{
+		UE_LOG(LogSceneState, Warning, TEXT("Failed to resolve bindings. Try compiling the template data again."));
+	}
+}
+
+void FSceneStateTemplateData::Link()
+{
+	for (TEnumerateRef<FSceneStateTransition> Transition : EnumerateRange(Transitions))
+	{
+		if (TransitionLinks.IsValidIndex(Transition.GetIndex()))
+		{
+			Transition->Link(TransitionLinks[Transition.GetIndex()], OwnerClass);
+		}
+	}
+
+	for (TEnumerateRef<FSceneStateConduit> Conduit : EnumerateRange(Conduits))
+	{
+		if (ConduitLinks.IsValidIndex(Conduit.GetIndex()))
+		{
+			Conduit->Link(ConduitLinks[Conduit.GetIndex()], OwnerClass);
+		}
+	}
+}
+
+void FSceneStateTemplateData::Reset()
+{
+	Tasks.Reset();
+	TaskPrerequisites.Reset();
+	TaskInstances.Reset();
+	Functions.Reset();
+	FunctionInstances.Reset();
+	States.Reset();
+	Conduits.Reset();
+	ConduitLinks.Reset();
+	EventHandlers.Reset();
+	Transitions.Reset();
+	TransitionLinks.Reset();
+	TransitionParameters.Reset();
+	StateMachines.Reset();
+	BindingCollection.Reset();
+	StateMachineIdToIndex.Reset();
+
+#if WITH_EDITOR
+	TaskMetadata.Reset();
+	FunctionMetadata.Reset();
+	StateMetadata.Reset();
+	TransitionMetadata.Reset();
+	StateNodeToIndex.Reset();
+	StateMachineGraphToIndex.Reset();
+	TaskNodeToIndex.Reset();
+#endif
+}
+
+#if WITH_EDITOR
+const FSceneState* FSceneStateTemplateData::FindStateFromNode(FObjectKey InStateNode) const
+{
+	const uint16* StateIndex = StateNodeToIndex.Find(InStateNode);
+	if (StateIndex && States.IsValidIndex(*StateIndex))
+	{
+		return &States[*StateIndex];
+	}
+	return nullptr;
+}
+
+const FSceneStateTask* FSceneStateTemplateData::FindTaskFromNode(FObjectKey InTaskNode) const
+{
+	const uint16* TaskIndex = TaskNodeToIndex.Find(InTaskNode);
+	if (TaskIndex && Tasks.IsValidIndex(*TaskIndex))
+	{
+		return Tasks[*TaskIndex].GetPtr<const FSceneStateTask>();
+	}
+	return nullptr;
+}
+
+void FSceneStateTemplateData::ForEachStateInstance(const TSharedRef<UE::SceneState::FExecutionContextRegistry>& InContextRegistry
+	, FObjectKey InStateNode
+	, TFunctionRef<void(const FSceneStateInstance&)> InFunctor) const
+{
+	const uint16* StateIndex = StateNodeToIndex.Find(InStateNode);
+	if (!StateIndex)
+	{
+		return;
+	}
+
+	InContextRegistry->ForEachExecutionContext(
+		[StateIndex, &InFunctor](const FSceneStateExecutionContext& InExecutionContext)
+		{
+			if (const FSceneStateInstance* StateInstance = InExecutionContext.FindStateInstance(*StateIndex))
+			{
+				InFunctor(*StateInstance);
+			}
+		});
+}
+
+void FSceneStateTemplateData::ForEachTaskInstance(const TSharedRef<UE::SceneState::FExecutionContextRegistry>& InContextRegistry, FObjectKey InTaskNode, TFunctionRef<void(const FSceneStateTaskInstance&)> InFunctor) const
+{
+	const uint16* TaskIndex = TaskNodeToIndex.Find(InTaskNode);
+	if (!TaskIndex)
+	{
+		return;
+	}
+
+	InContextRegistry->ForEachExecutionContext(
+		[TaskIndex, &InFunctor](const FSceneStateExecutionContext& InExecutionContext)
+		{
+			if (const FSceneStateTaskInstance* TaskInstance = InExecutionContext.FindTaskInstance(*TaskIndex).GetPtr<const FSceneStateTaskInstance>())
+			{
+				InFunctor(*TaskInstance);
+			}
+		});
+}
+
+void FSceneStateTemplateData::OnObjectsReinstanced(const TMap<UObject*, UObject*>& InReplacementMap)
+{
+	TSet<const UStruct*> Structs;
+	Structs.Reserve(InReplacementMap.Num());
+
+	bool bRequiresResolve = false;
+
+	for (const TPair<UObject*, UObject*>& Pair : InReplacementMap)
+	{
+		if (const UObject* Replacement = Pair.Value)
+		{
+			if (Replacement->IsIn(OwnerClass))
+			{
+				bRequiresResolve = true;
+				break;
+			}
+			Structs.Add(Replacement->GetClass());
+		}
+	}
+
+	if (bRequiresResolve || BindingCollection.ContainsAnyStruct(Structs))
+	{
+		ResolveBindings();
+	}
+}
+
+void FSceneStateTemplateData::OnStructsReinstanced(const UUserDefinedStruct& InStruct)
+{
+	TSet<const UStruct*> Structs;
+	Structs.Add(&InStruct);
+
+	if (BindingCollection.ContainsAnyStruct(Structs))
+	{
+		ResolveBindings();
+	}
+}
+#endif
diff --git a/Engine/Plugins/VirtualProduction/SceneState/Source/SceneState/Private/SceneStateUtils.cpp b/Engine/Plugins/VirtualProduction/SceneState/Source/SceneState/Private/SceneStateUtils.cpp
--- a/Engine/Plugins/VirtualProduction/SceneState/Source/SceneState/Private/SceneStateUtils.cpp
+++ b/Engine/Plugins/VirtualProduction/SceneState/Source/SceneState/Private/SceneStateUtils.cpp
@@ -82,16 +82,11 @@ bool ApplyBatch(const FSceneStateExecutionContext& InContext, const FApplyBatchP
 		return false;
 	}
 
-	const FSceneStateBindingCollection* const BindingCollection = InContext.GetBindingCollection();
-	if (!BindingCollection)
-	{
-		UE_LOG(LogSceneState, Error, TEXT("[%s] ApplyBatch failed. Binding Collection is invalid."), *InContext.GetExecutionContextName())
-		return false;
-	}
+	const FSceneStateBindingCollection& BindingCollection = InContext.GetBindingCollection();
 
 	bool bResult = true;
 
-	const FPropertyBindingCopyInfoBatch& Batch = BindingCollection->GetBatch(BindingsBatch);
+	const FPropertyBindingCopyInfoBatch& Batch = BindingCollection.GetBatch(BindingsBatch);
 	check(InParams.TargetDataView.GetStruct() == Batch.TargetStruct.Get().Struct);
 
 	// If there were valid functions found on setup, execute them now
@@ -109,10 +104,10 @@ bool ApplyBatch(const FSceneStateExecutionContext& InContext, const FApplyBatchP
 		}
 	}
 
-	for (const FPropertyBindingCopyInfo& Copy : BindingCollection->GetBatchCopies(Batch))
+	for (const FPropertyBindingCopyInfo& Copy : BindingCollection.GetBatchCopies(Batch))
 	{
 		const FPropertyBindingDataView SourceView = InContext.FindDataView(Copy.SourceDataHandle.Get<FSceneStateBindingDataHandle>());
-		bResult &= BindingCollection->CopyProperty(Copy, SourceView, InParams.TargetDataView);
+		bResult &= BindingCollection.CopyProperty(Copy, SourceView, InParams.TargetDataView);
 	}
 
 	return bResult;
diff --git a/Engine/Plugins/VirtualProduction/SceneState/Source/SceneState/Private/Tasks/SceneStateTaskExecutionContext.cpp b/Engine/Plugins/VirtualProduction/SceneState/Source/SceneState/Private/Tasks/SceneStateTaskExecutionContext.cpp
--- a/Engine/Plugins/VirtualProduction/SceneState/Source/SceneState/Private/Tasks/SceneStateTaskExecutionContext.cpp
+++ b/Engine/Plugins/VirtualProduction/SceneState/Source/SceneState/Private/Tasks/SceneStateTaskExecutionContext.cpp
@@ -15,7 +15,7 @@ namespace Private
 
 FExecutionContextHandle GetContextHandle(const FSceneStateExecutionContext& InContext)
 {
-	if (TSharedPtr<const FExecutionContextRegistry> ContextRegistry = InContext.GetContextRegistry().Pin())
+	if (TSharedPtr<const FExecutionContextRegistry> ContextRegistry = InContext.GetContextRegistry())
 	{
 		return ContextRegistry->FindHandle(InContext);
 	}
diff --git a/Engine/Plugins/VirtualProduction/SceneState/Source/SceneState/Private/Transition/SceneStateTransition.cpp b/Engine/Plugins/VirtualProduction/SceneState/Source/SceneState/Private/Transition/SceneStateTransition.cpp
--- a/Engine/Plugins/VirtualProduction/SceneState/Source/SceneState/Private/Transition/SceneStateTransition.cpp
+++ b/Engine/Plugins/VirtualProduction/SceneState/Source/SceneState/Private/Transition/SceneStateTransition.cpp
@@ -1,20 +1,23 @@
 // Copyright Epic Games, Inc. All Rights Reserved.
 
 #include "Transition/SceneStateTransition.h"
+#include "Conduit/SceneStateConduit.h"
+#include "SceneState.h"
+#include "SceneStateEventHandler.h"
+#include "SceneStateEventStream.h"
 #include "SceneStateExecutionContext.h"
-#include "SceneStateGeneratedClass.h"
-#include "SceneStateObject.h"
+#include "SceneStateMachine.h"
 #include "SceneStateUtils.h"
 #include "Transition/SceneStateTransitionEvaluation.h"
 #include "Transition/SceneStateTransitionInstance.h"
 #include "Transition/SceneStateTransitionLink.h"
 #include "Transition/SceneStateTransitionResult.h"
 
-void FSceneStateTransition::Link(const FSceneStateTransitionLink& InTransitionLink, USceneStateGeneratedClass* InGeneratedClass)
+void FSceneStateTransition::Link(const FSceneStateTransitionLink& InTransitionLink, TNotNull<const UClass*> InOwnerClass)
 {
 	if (!InTransitionLink.ResultPropertyName.IsNone())
 	{
-		ResultProperty = CastField<FStructProperty>(InGeneratedClass->FindPropertyByName(InTransitionLink.ResultPropertyName));
+		ResultProperty = CastField<FStructProperty>(InOwnerClass->FindPropertyByName(InTransitionLink.ResultPropertyName));
 		check(ResultProperty);
 	}
 	else
@@ -24,7 +27,7 @@ void FSceneStateTransition::Link(const FSceneStateTransitionLink& InTransitionLi
 
 	if (!InTransitionLink.EventName.IsNone())
 	{
-		EvaluationEvent = InGeneratedClass->FindFunctionByName(InTransitionLink.EventName);
+		EvaluationEvent = InOwnerClass->FindFunctionByName(InTransitionLink.EventName);
 		check(EvaluationEvent);
 	}
 	else
@@ -164,11 +167,15 @@ bool FSceneStateTransition::ProcessEvaluationEvent(const UE::SceneState::FTransi
 		FunctionParams = Instance->Parameters.GetMutableValue().GetMemory();
 	}
 
-	USceneStateObject* RootState = InParams.ExecutionContext.GetRootState();
-	check(RootState);
-	RootState->ProcessEvent(EvaluationEvent, FunctionParams);
+	UObject* RootObject = InParams.ExecutionContext.GetRootObject();
+	if (!ensure(RootObject))
+	{
+		return false;
+	}
+
+	RootObject->ProcessEvent(EvaluationEvent, FunctionParams);
 
-	const FSceneStateTransitionResult* Result = ResultProperty->ContainerPtrToValuePtr<FSceneStateTransitionResult>(RootState);
+	const FSceneStateTransitionResult* Result = ResultProperty->ContainerPtrToValuePtr<FSceneStateTransitionResult>(RootObject);
 	check(Result);
 	return Result->bCanTransition;
 }
diff --git a/Engine/Plugins/VirtualProduction/SceneState/Source/SceneState/Public/Conduit/SceneStateConduit.h b/Engine/Plugins/VirtualProduction/SceneState/Source/SceneState/Public/Conduit/SceneStateConduit.h
--- a/Engine/Plugins/VirtualProduction/SceneState/Source/SceneState/Public/Conduit/SceneStateConduit.h
+++ b/Engine/Plugins/VirtualProduction/SceneState/Source/SceneState/Public/Conduit/SceneStateConduit.h
@@ -7,7 +7,6 @@
 #include "UObject/ObjectPtr.h"
 #include "SceneStateConduit.generated.h"
 
-class USceneStateGeneratedClass;
 struct FSceneStateConduitLink;
 struct FSceneStateExecutionContext;
 
@@ -28,8 +27,8 @@ struct FSceneStateConduit
 {
 	GENERATED_BODY()
 
-	/** Called when the Generated Class is linking to cache the Event function and Result property */
-	void Link(const FSceneStateConduitLink& InConduitLink, USceneStateGeneratedClass* InGeneratedClass);
+	/** Called when the Owner Class is linking to cache the Event function and Result property */
+	void Link(const FSceneStateConduitLink& InConduitLink, TNotNull<const UClass*> InOwnerClass);
 
 	/** Called when a state that could evaluate this conduit has entered. Allocates and initializes the conduit instance */
 	void Setup(const FSceneStateExecutionContext& InContext) const;
diff --git a/Engine/Plugins/VirtualProduction/SceneState/Source/SceneState/Public/SceneStateExecutionContext.h b/Engine/Plugins/VirtualProduction/SceneState/Source/SceneState/Public/SceneStateExecutionContext.h
--- a/Engine/Plugins/VirtualProduction/SceneState/Source/SceneState/Public/SceneStateExecutionContext.h
+++ b/Engine/Plugins/VirtualProduction/SceneState/Source/SceneState/Public/SceneStateExecutionContext.h
@@ -14,9 +14,10 @@
 #include "UObject/ObjectPtr.h"
 #include "SceneStateExecutionContext.generated.h"
 
+#define UE_API SCENESTATE_API
+
 class UClass;
 class USceneStateEventStream;
-class USceneStateGeneratedClass;
 class USceneStateObject;
 class USceneStatePlayer;
 struct FPropertyBindingDataView;
@@ -31,6 +32,7 @@ struct FSceneStateMachine;
 struct FSceneStateMetadata;
 struct FSceneStatePropertyReference;
 struct FSceneStateTask;
+struct FSceneStateTemplateData;
 struct FSceneStateTransition;
 
 namespace UE::SceneState
@@ -56,53 +58,70 @@ struct FSceneStateExecutionContext
 {
 	GENERATED_BODY()
 
-	SCENESTATE_API static const FSceneStateExecutionContext InvalidContext;
+	UE_API static const FSceneStateExecutionContext InvalidContext;
 
-	SCENESTATE_API ~FSceneStateExecutionContext();
+	UE_API ~FSceneStateExecutionContext();
 
-	const USceneStateGeneratedClass* GetGeneratedClass() const
-	{
-		return GeneratedClass;
-	}
-
-	USceneStateObject* GetRootState() const
-	{
-		return RootState;
-	}
+	UE_API const FSceneStateTemplateData& GetSharedTemplateData() const;
 
 	/** Gets the registry that this context is registered to */
-	TWeakPtr<const UE::SceneState::FExecutionContextRegistry> GetContextRegistry() const;
+	TSharedPtr<const UE::SceneState::FExecutionContextRegistry> GetContextRegistry() const;
 
 	/** Gets the binding collection owned by the generated class */
-	const FSceneStateBindingCollection* GetBindingCollection() const;
+	const FSceneStateBindingCollection& GetBindingCollection() const;
 
 	/** Returns the Player Debug Name for logging purposes */
-	SCENESTATE_API FString GetExecutionContextName() const;
+	UE_API const FString& GetExecutionContextName() const;
+
+	/** Retrieves the Root Object of this Execution */
+	UE_API UObject* GetRootObject() const;
 
 	/** Retrieves the Context Object for this Execution */
-	SCENESTATE_API UObject* GetContextObject() const;
+	UE_API UObject* GetContextObject() const;
 
 	/** Gets the Event Stream from the Root State if available */
-	SCENESTATE_API USceneStateEventStream* GetEventStream() const;
+	UE_API USceneStateEventStream* GetEventStream() const;
 
 	/**
 	 * Initialization for a context with a parent context.
 	 * @param InParentContext the parent of this execution context
 	 */
-	SCENESTATE_API void Setup(const FSceneStateExecutionContext& InParentContext);
+	UE_API void Setup(const FSceneStateExecutionContext& InParentContext);
+
+	/**
+	 * Initialization for the Context.
+	 * Through the given root state it pre-allocates the instance data that will be used
+	 * @param InRootState the parameters containing the template data and context for this execution
+	 */
+	UE_API void Setup(TNotNull<USceneStateObject*> InRootState);
 
+	struct FSetupParams
+	{
+		/** Shared memory of the template data holding all the state machines, states, etc. */
+		FConstSharedStruct SharedTemplateData;
+		/** Root object instance */
+		UObject* RootObject = nullptr;
+		/** Context object for utility purposes (e.g. to get world from context, etc) */
+		UObject* ContextObject = nullptr;
+		/** Event stream to use for execution */
+		USceneStateEventStream* EventStream = nullptr;
+		/** Context registry to use */
+		TSharedPtr<UE::SceneState::FExecutionContextRegistry> ContextRegistry;
+		/** Debug name describing the context */
+		FString ContextName;
+	};
 	/**
 	 * Initialization for the Context.
 	 * Through the given root state it pre-allocates the instance data that will be used
-	 * @param InRootState the root state of this execution
+	 * @param InParams the parameters containing the template data and context for this execution
 	 */
-	SCENESTATE_API void Setup(USceneStateObject* InRootState);
+	UE_API void Setup(const FSetupParams& InParams);
 
 	/** Called when cleaning up the instances of this execution */
-	SCENESTATE_API void Reset();
+	UE_API void Reset();
 
 	/** Finds the data view that matches the given data handle */
-	SCENESTATE_API FPropertyBindingDataView FindDataView(const FSceneStateBindingDataHandle& InDataHandle) const;
+	UE_API FPropertyBindingDataView FindDataView(const FSceneStateBindingDataHandle& InDataHandle) const;
 
 	/** Returns the pointer to the property value address the given reference is mapped to */
 	template<typename T = void>
@@ -244,24 +263,36 @@ private:
 	const FSceneStateExecutionContext* FindParentContext() const;
 
 	/** Retrieves the Absolute Index of a State in the State array */
-	bool GetStateIndex(const FSceneState& InState, uint16& OutIndex) const;
+	uint16 GetStateIndex(const FSceneStateTemplateData& InTemplateData, const FSceneState& InState) const;
 
 	/** Retrieves the Absolute Index of a State Machine in the State Machine array */
-	bool GetStateMachineIndex(const FSceneStateMachine& InStateMachine, uint16& OutIndex) const;
+	uint16 GetStateMachineIndex(const FSceneStateTemplateData& InTemplateData, const FSceneStateMachine& InStateMachine) const;
 
 	/** Retrieves the Absolute Index of a Transition in the Transitions array */
-	bool GetTransitionIndex(const FSceneStateTransition& InTransition, uint16& OutIndex) const;
+	uint16 GetTransitionIndex(const FSceneStateTemplateData& InTemplateData, const FSceneStateTransition& InTransition) const;
 
 	/** Retrieves the Absolute Index of a Conduit in the Conduits array */
-	bool GetConduitIndex(const FSceneStateConduit& InConduit, uint16& OutIndex) const;
+	uint16 GetConduitIndex(const FSceneStateTemplateData& InTemplateData, const FSceneStateConduit& InConduit) const;
+
+	/** Weak ptr to the context holding this execution */
+	UPROPERTY()
+	TWeakObjectPtr<UObject> ContextObjectWeak;
 
-	/** Root state object owning the Scene State Execution */
+	/** Weak ptr to the event stream for this execution */
 	UPROPERTY()
-	TObjectPtr<USceneStateObject> RootState = nullptr;
+	TWeakObjectPtr<USceneStateEventStream> EventStreamWeak;
 
-	/** Class of the root state */
+	/** Weak ptr to the root object owning this execution */
 	UPROPERTY()
-	TObjectPtr<USceneStateGeneratedClass> GeneratedClass = nullptr;
+	TObjectPtr<UObject> RootObject;
+
+	/** Description of the context */
+	UPROPERTY()
+	FString ContextName;
+
+	/** Scene State Template */
+	UPROPERTY()
+	FConstSharedStruct SharedTemplateData;
 
 	/**
 	 * Map of State Index to its Instance Data
@@ -322,3 +353,5 @@ private:
 
 	friend class UE::SceneState::FExecutionScope;
 };
+
+#undef UE_API
diff --git a/Engine/Plugins/VirtualProduction/SceneState/Source/SceneState/Public/SceneStateGeneratedClass.h b/Engine/Plugins/VirtualProduction/SceneState/Source/SceneState/Public/SceneStateGeneratedClass.h
--- a/Engine/Plugins/VirtualProduction/SceneState/Source/SceneState/Public/SceneStateGeneratedClass.h
+++ b/Engine/Plugins/VirtualProduction/SceneState/Source/SceneState/Public/SceneStateGeneratedClass.h
@@ -2,27 +2,14 @@
 
 #pragma once
 
-#include "Conduit/SceneStateConduit.h"
-#include "Conduit/SceneStateConduitLink.h"
 #include "Engine/BlueprintGeneratedClass.h"
-#include "Functions/SceneStateFunctionMetadata.h"
-#include "SceneState.h"
-#include "SceneStateBindingCollection.h"
-#include "SceneStateEventHandler.h"
-#include "SceneStateMachine.h"
-#include "SceneStateRange.h"
-#include "StructUtils/InstancedStructContainer.h"
-#include "Tasks/SceneStateTask.h"
-#include "Tasks/SceneStateTaskMetadata.h"
-#include "Transition/SceneStateTransition.h"
-#include "Transition/SceneStateTransitionLink.h"
-#include "Transition/SceneStateTransitionMetadata.h"
-#include "UObject/ObjectKey.h"
-#include "UObject/Package.h"
+#include "StructUtils/SharedStruct.h"
 #include "SceneStateGeneratedClass.generated.h"
 
-class USceneStateObject;
+#define UE_API SCENESTATE_API
+
 class UUserDefinedStruct;
+struct FSceneStateTemplateData;
 
 namespace UE::SceneState::Editor
 {
@@ -34,8 +21,8 @@ namespace UE::SceneState::Editor
 
 /**
  * Object Class for the Scene State Object
- * Holds all the data about States, State Machines, Tasks, etc.
- * All this data is immutable in execution, and as such, it is not instanced to the Scene State Object instances.
+ * Holds the template data for scene state objects.
+ * @see FSceneStateTemplateData
  * @see FSceneStateExecutionContext
  */
 UCLASS(MinimalAPI)
@@ -44,131 +31,25 @@ class USceneStateGeneratedClass : public UBlueprintGeneratedClass
 	GENERATED_BODY()
 
 public:
-	SCENESTATE_API USceneStateGeneratedClass();
-
-	/** Gets the state that owns all the "top-level" state machines */
-	const FSceneState* GetRootState() const;
-
-	/** Finds the Top-Level State Machine that is linked to the given Guid */
-	const FSceneStateMachine* FindStateMachine(const FGuid& InStateMachineId) const;
-
-	const FSceneStateBindingCollection& GetBindingCollection() const
-	{
-		return BindingCollection;
-	}
-
-	TConstArrayView<FSceneState> GetStates() const
-	{
-		return States;
-	}
-
-	TConstArrayView<FSceneStateConduit> GetConduits() const
-	{
-		return Conduits;
-	}
-
-	const FInstancedStructContainer& GetTasks() const
-	{
-		return Tasks;
-	}
-
-	TConstArrayView<uint16> GetTaskPrerequisites() const
-	{
-		return TaskPrerequisites;
-	}
-
-	const FInstancedStructContainer& GetTaskInstances() const
-	{
-		return TaskInstances;
-	}
-
-	TConstArrayView<FSceneStateEventHandler> GetEventHandlers() const
-	{
-		return EventHandlers;
-	}
-
-	TConstArrayView<FSceneStateTransition> GetTransitions() const
-	{
-		return Transitions;
-	}
-
-	/** Gets a copy of the transition parameters at the given transition absolute index. */
-	FInstancedPropertyBag GetTransitionParameters(uint16 InAbsoluteIndex) const;
-
-	TConstArrayView<FSceneStateMachine> GetStateMachines() const
-	{
-		return StateMachines;
-	}
-
-	const FInstancedStructContainer& GetFunctions() const
-	{
-		return Functions;
-	}
-
-	const FInstancedStructContainer& GetFunctionInstances() const
-	{
-		return FunctionInstances;
-	}
-
-#if WITH_EDITOR
-	TConstArrayView<FSceneStateMetadata> GetStateMetadata() const
-	{
-		return StateMetadata;
-	}
+	UE_API USceneStateGeneratedClass();
 
-	TConstArrayView<FSceneStateTaskMetadata> GetTaskMetadata() const
-	{
-		return TaskMetadata;
-	}
-#endif
+	UE_API FConstSharedStruct GetSharedTemplateData() const;
 
 	//~ Begin UClass
-	SCENESTATE_API virtual void Link(FArchive& Ar, bool bInRelinkExistingProperties) override;
-	SCENESTATE_API virtual void PurgeClass(bool bInRecompilingOnLoad) override;
+	UE_API virtual void Link(FArchive& Ar, bool bInRelinkExistingProperties) override;
+	UE_API virtual void PurgeClass(bool bInRecompilingOnLoad) override;
 	//~ End UClass
 
 	//~ Begin UObject
-	SCENESTATE_API virtual void PostLoad() override;
-	SCENESTATE_API virtual void BeginDestroy() override;
+	UE_API virtual void PostLoad() override;
+	UE_API virtual void BeginDestroy() override;
 	//~ End UObject
 
-#if WITH_EDITOR
-	/** Finds the mapped compiled state for a given state node */
-	SCENESTATE_API const FSceneState* FindStateFromNode(FObjectKey InStateNode) const;
-
-	/** Finds the mapped compiled task for a given task node */
-	SCENESTATE_API const FSceneStateTask* FindTaskFromNode(FObjectKey InTaskNode) const;
-
-	/**
-	 * For a given root state and state node retrieves the mapped state instances
-	 * @param InRootState the root state that is outered to a player
-	 * @param InStateNode the state node to look for 
-	 * @param InFunctor the functor to run for each state instance
-	 */
-	SCENESTATE_API void ForEachStateInstance(USceneStateObject& InRootState
-		, FObjectKey InStateNode
-		, TFunctionRef<void(const FSceneStateInstance&)> InFunctor) const;
-
-	/**
-	 * For a given root state and task node retrieves the mapped task instances
-	 * @param InRootState the root state that is outered to a player
-	 * @param InTaskNode the task node to look for
-	 * @param InFunctor the functor to run for each task instance
-	 */
-	SCENESTATE_API void ForEachTaskInstance(USceneStateObject& InRootState
-		, FObjectKey InTaskNode
-		, TFunctionRef<void(const FSceneStateTaskInstance&)> InFunctor) const;
-#endif
-
 private:
-	/** Finds the struct type for the given data handle */
-	const UStruct* FindDataStruct(const FSceneStateBindingDataHandle& InDataHandle);
+	FSceneStateTemplateData& GetTemplateData();
 
-	/** Patches bindings and resolves the binding paths for the owning binding collection */
-	SCENESTATE_API void ResolveBindings();
-
-	/** Resets all the elements that get compiled for this generated class (e.g. Tasks, State Machines, etc) */
-	void Reset();
+	/** Called to resolve the template data bindings */
+	UE_API void ResolveBindings();
 
 	bool IsFullClass() const;
 
@@ -177,93 +58,11 @@ private:
 	void OnStructsReinstanced(const UUserDefinedStruct& InStruct);
 #endif
 
+	/** Scene State template data (FSceneStateTemplateData) for this class. */
 	UPROPERTY()
-	FSceneStateBindingCollection BindingCollection;
-
-	UPROPERTY()
-	uint16 RootStateIndex = FSceneStateRange::InvalidIndex;
-
-	/** Holds all the compiled states including nested states */
-	UPROPERTY()
-	TArray<FSceneState> States;
-
-	/** Holds all the compiled tasks (including tasks from nested states) in a contiguous block of memory */
-	UPROPERTY()
-	FInstancedStructContainer Tasks;
-
-	/** Array of the task prerequisites in their relative index. Each task has a unique range indicating their pre-requisites */
-	UPROPERTY()
-	TArray<uint16> TaskPrerequisites;
-
-	/** Templates used to instantiate the Task Instances */
-	UPROPERTY()
-	FInstancedStructContainer TaskInstances;
-
-	/** Holds all the compiled functions (including from nested sources) in a contiguous block of memory */
-	UPROPERTY()
-	FInstancedStructContainer Functions;
-
-	/** Templates used to instantiate the function data */
-	UPROPERTY()
-	FInstancedStructContainer FunctionInstances;
-
-	/** Holds all the compiled event handlers (including from nested sources) contiguously */
-	UPROPERTY()
-	TArray<FSceneStateEventHandler> EventHandlers;
-
-	/** Holds all the exit transitions (including nested sources) contiguously */
-	UPROPERTY()
-	TArray<FSceneStateTransition> Transitions;
-
-	/** Compiled transition information only used in Link time */
-	UPROPERTY()
-	TArray<FSceneStateTransitionLink> TransitionLinks;
-
-	/** Map of the Transition Index (absolute) to the template transition parameters map for evaluation function call. */
-	UPROPERTY()
-	TMap<uint16, FInstancedPropertyBag> TransitionParameters;
-
-	/** All the compiled conduits */
-	UPROPERTY()
-	TArray<FSceneStateConduit> Conduits;
-
-	/** Compiled conduit information used only in Link time */
-	UPROPERTY()
-	TArray<FSceneStateConduitLink> ConduitLinks;
-
-	UPROPERTY()
-	TArray<FSceneStateMachine> StateMachines;
-
-	/** Map of the Top-Level State machine Parameters id to the index in the state machine array */
-	UPROPERTY()
-	TMap<FGuid, uint16> StateMachineIdToIndex;
-
-#if WITH_EDITORONLY_DATA
-	/** Holds all the state editor-only metadata. This array matches in index with the states array. */
-	UPROPERTY()
-	TArray<FSceneStateMetadata> StateMetadata;
-
-	/** Holds all the task editor-only metadata. This array matches in index with the tasks array. */
-	UPROPERTY()
-	TArray<FSceneStateTaskMetadata> TaskMetadata;
-
-	/** Holds all the function editor-only metadata. This array matches in index with the functions. */
-	UPROPERTY()
-	TArray<FSceneStateFunctionMetadata> FunctionMetadata;
-
-	/** Holds all the transition editor-only metadata. This array matches in index with the transitions array. */
-	UPROPERTY()
-	TArray<FSceneStateTransitionMetadata> TransitionMetadata;
-
-	/** Map of a State Node to its State Index */
-	TMap<FObjectKey, uint16> StateNodeToIndex;
-
-	/** Map of a State Machine Graph to its State Machine Index */
-	TMap<FObjectKey, uint16> StateMachineGraphToIndex;
-
-	/** Map of a Task Node to its Task Index */
-	TMap<FObjectKey, uint16> TaskNodeToIndex;
+	FSharedStruct SharedTemplateData;
 
+#if WITH_EDITOR
 	FDelegateHandle OnObjectsReinstancedHandle;
 	FDelegateHandle OnStructsReinstancedHandle;
 #endif
@@ -274,3 +73,5 @@ private:
 	friend UE::SceneState::Editor::FStateMachineCompiler;
 	friend class USceneStateBlueprint;
 };
+
+#undef UE_API
diff --git a/Engine/Plugins/VirtualProduction/SceneState/Source/SceneState/Public/SceneStateObject.h b/Engine/Plugins/VirtualProduction/SceneState/Source/SceneState/Public/SceneStateObject.h
--- a/Engine/Plugins/VirtualProduction/SceneState/Source/SceneState/Public/SceneStateObject.h
+++ b/Engine/Plugins/VirtualProduction/SceneState/Source/SceneState/Public/SceneStateObject.h
@@ -11,9 +11,10 @@
 #include "UObject/Object.h"
 #include "SceneStateObject.generated.h"
 
+#define UE_API SCENESTATE_API
+
 class USceneStateBlueprintableTask;
 class USceneStateEventStream;
-class USceneStateGeneratedClass;
 struct FPropertyBindingDataView;
 struct FSceneStateBindingDataHandle;
 struct FSceneStateExecutionContext;
@@ -33,17 +34,18 @@ class USceneStateObject : public UObject
 {
 	GENERATED_BODY()
 
-	friend USceneStateGeneratedClass;
-
 public:
-	SCENESTATE_API USceneStateObject();
+	UE_API USceneStateObject();
 
 	/** Gets the context name of the scene state player */
-	SCENESTATE_API FString GetContextName() const;
+	UE_API FString GetContextName() const;
+
+	/** Gets the template data (FSceneStateTemplate) for the given object */
+	UE_API FConstSharedStruct GetSharedTemplateData() const;
 
 	/** Gets the context object of the scene state player */
 	UFUNCTION(BlueprintCallable, Category="State")
-	SCENESTATE_API UObject* GetContextObject() const;
+	UE_API UObject* GetContextObject() const;
 
 	UFUNCTION(BlueprintPure, Category="Events")
 	USceneStateEventStream* GetEventStream() const
@@ -53,7 +55,7 @@ public:
 
 	/** Returns true if the object has an active root state */
 	UFUNCTION(BlueprintPure, Category="State")
-	SCENESTATE_API bool IsActive() const;
+	UE_API bool IsActive() const;
 
 	/** Called by the scene state player to setup this object and the root execution context */
 	void Setup();
@@ -76,17 +78,21 @@ public:
 	UFUNCTION(BlueprintImplementableEvent, DisplayName = "Exit", Category = "State Events")
 	void ReceiveExit();
 
-	TSharedRef<UE::SceneState::FExecutionContextRegistry> GetContextRegistry() const;
+	/** Gets the context registry owned by this object */
+	UE_API TSharedRef<UE::SceneState::FExecutionContextRegistry> GetContextRegistry() const;
 
 	//~ Begin UObject
-	SCENESTATE_API virtual UWorld* GetWorld() const override final;
-	SCENESTATE_API virtual void BeginDestroy() override;
+	UE_API virtual UWorld* GetWorld() const override final;
+	UE_API virtual void BeginDestroy() override;
 	//~ End UObject
 
 private:
-	/** The class of this scene state object */
+	/** Resolves the root state from the shared template data, if any */
+	const FSceneState* GetRootState() const;
+
+	/** Scene State template data (FSceneStateTemplateData) to use for this object's execution. */
 	UPROPERTY(Transient, DuplicateTransient, TextExportTransient, NonTransactional, meta=(NoBinding))
-	TObjectPtr<USceneStateGeneratedClass> GeneratedClass;
+	FConstSharedStruct SharedTemplateData;
 
 	/** Runtime Event System keeping track of Active Events that have been added */
 	UPROPERTY(Transient, Instanced, meta=(NoBinding))
@@ -99,3 +105,5 @@ private:
 	/** The registry containing all the execution contexts in this object */
 	TSharedPtr<UE::SceneState::FExecutionContextRegistry> ContextRegistry;
 };
+
+#undef UE_API
diff --git a/Engine/Plugins/VirtualProduction/SceneState/Source/SceneState/Public/SceneStateTemplateData.h b/Engine/Plugins/VirtualProduction/SceneState/Source/SceneState/Public/SceneStateTemplateData.h
new file mode 100644
--- /dev/null
+++ b/Engine/Plugins/VirtualProduction/SceneState/Source/SceneState/Public/SceneStateTemplateData.h
@@ -0,0 +1,209 @@
+// Copyright Epic Games, Inc. All Rights Reserved.
+
+#pragma once
+
+#include "Conduit/SceneStateConduit.h"
+#include "Conduit/SceneStateConduitLink.h"
+#include "Containers/Array.h"
+#include "SceneState.h"
+#include "SceneStateBindingCollection.h"
+#include "SceneStateEventHandler.h"
+#include "SceneStateMachine.h"
+#include "SceneStateRange.h"
+#include "StructUtils/InstancedStructContainer.h"
+#include "StructUtils/PropertyBag.h"
+#include "Transition/SceneStateTransition.h"
+#include "Transition/SceneStateTransitionLink.h"
+#include "UObject/ObjectKey.h"
+
+#if WITH_EDITOR
+#include "Functions/SceneStateFunctionMetadata.h"
+#include "Tasks/SceneStateTaskMetadata.h"
+#include "Transition/SceneStateTransitionMetadata.h"
+#endif
+
+#include "SceneStateTemplateData.generated.h"
+
+#define UE_API SCENESTATE_API
+
+namespace UE::SceneState
+{
+	class FExecutionContextRegistry;
+}
+
+struct FSceneStateInstance;
+struct FSceneStateTask;
+struct FSceneStateTaskInstance;
+
+/**
+ * Holds all the data about States, State Machines, Tasks, etc.
+ * All this data is immutable in execution, and as such, it is not instanced to the Scene State Object instances.
+ * @see FSceneStateExecutionContext
+ */
+USTRUCT()
+struct FSceneStateTemplateData
+{
+	GENERATED_BODY()
+
+	/** For compatibility only. Do not use */
+	explicit FSceneStateTemplateData(EForceInit)
+	{
+	}
+
+	/** Constructs the template data for a given owner class */
+	explicit FSceneStateTemplateData(TNotNull<const UClass*> InOwnerClass);
+
+	/** Gets the state that contains all top-level state machines */
+	const FSceneState* GetRootState() const;
+
+	/** Finds the state machine for a given state machine guid */
+	const FSceneStateMachine* FindStateMachine(const FGuid& InStateMachineId) const;
+
+	/** Finds the struct type for the given data handle */
+	const UStruct* FindDataStruct(const FSceneStateBindingDataHandle& InDataHandle);
+
+	/** Patches bindings and resolves the binding paths for the owning binding collection */
+	UE_API void ResolveBindings();
+
+	/** Called at link time to cache the properties and functions to use (e.g. in transitions or conduits) */
+	UE_API void Link();
+
+	/** Resets all the elements that get compiled for this generated class (e.g. Tasks, State Machines, etc) */
+	UE_API void Reset();
+
+#if WITH_EDITOR
+	/** Finds the mapped compiled state for a given state node */
+	UE_API const FSceneState* FindStateFromNode(FObjectKey InStateNode) const;
+
+	/** Finds the mapped compiled task for a given task node */
+	UE_API const FSceneStateTask* FindTaskFromNode(FObjectKey InTaskNode) const;
+
+	/**
+	 * For a given root state and state node retrieves the mapped state instances
+	 * @param InContextRegistry the registry containing all the execution contexts
+	 * @param InStateNode the state node to look for 
+	 * @param InFunctor the functor to run for each state instance
+	 */
+	UE_API void ForEachStateInstance(const TSharedRef<UE::SceneState::FExecutionContextRegistry>& InContextRegistry
+		, FObjectKey InStateNode
+		, TFunctionRef<void(const FSceneStateInstance&)> InFunctor) const;
+
+	/**
+	 * For a given root state and task node retrieves the mapped task instances
+	 * @param InContextRegistry the registry containing all the execution contexts
+	 * @param InTaskNode the task node to look for
+	 * @param InFunctor the functor to run for each task instance
+	 */
+	UE_API void ForEachTaskInstance(const TSharedRef<UE::SceneState::FExecutionContextRegistry>& InContextRegistry
+		, FObjectKey InTaskNode
+		, TFunctionRef<void(const FSceneStateTaskInstance&)> InFunctor) const;
+
+	void OnObjectsReinstanced(const TMap<UObject*, UObject*>& InReplacementMap);
+	void OnStructsReinstanced(const UUserDefinedStruct& InStruct);
+#endif
+
+	/** Owner class of this template data */
+	UPROPERTY()
+	TObjectPtr<const UClass> OwnerClass;
+
+	/** Container for all the bindings */
+	UPROPERTY()
+	FSceneStateBindingCollection BindingCollection;
+
+	/** Index to the state that contains all top-level state machines */
+	UPROPERTY()
+	uint16 RootStateIndex = FSceneStateRange::InvalidIndex;
+
+	/** Holds all the compiled states including nested states */
+	UPROPERTY()
+	TArray<FSceneState> States;
+
+	/** Holds all the compiled tasks (including tasks from nested states) in a contiguous block of memory */
+	UPROPERTY()
+	FInstancedStructContainer Tasks;
+
+	/** Array of the task prerequisites in their relative index. Each task has a unique range indicating their pre-requisites */
+	UPROPERTY()
+	TArray<uint16> TaskPrerequisites;
+
+	/** Templates used to instantiate the Task Instances */
+	UPROPERTY()
+	FInstancedStructContainer TaskInstances;
+
+	/** Holds all the compiled functions (including from nested sources) in a contiguous block of memory */
+	UPROPERTY()
+	FInstancedStructContainer Functions;
+
+	/** Templates used to instantiate the function data */
+	UPROPERTY()
+	FInstancedStructContainer FunctionInstances;
+
+	/** Holds all the compiled event handlers (including from nested sources) contiguously */
+	UPROPERTY()
+	TArray<FSceneStateEventHandler> EventHandlers;
+
+	/** Holds all the exit transitions (including nested sources) contiguously */
+	UPROPERTY()
+	TArray<FSceneStateTransition> Transitions;
+
+	/** Compiled transition information only used in Link time */
+	UPROPERTY()
+	TArray<FSceneStateTransitionLink> TransitionLinks;
+
+	/** Map of the Transition Index (absolute) to the template transition parameters map for evaluation function call. */
+	UPROPERTY()
+	TMap<uint16, FInstancedPropertyBag> TransitionParameters;
+
+	/** All the compiled conduits */
+	UPROPERTY()
+	TArray<FSceneStateConduit> Conduits;
+
+	/** Compiled conduit information used only in Link time */
+	UPROPERTY()
+	TArray<FSceneStateConduitLink> ConduitLinks;
+
+	UPROPERTY()
+	TArray<FSceneStateMachine> StateMachines;
+
+	/** Map of the Top-Level State machine Parameters id to the index in the state machine array */
+	UPROPERTY()
+	TMap<FGuid, uint16> StateMachineIdToIndex;
+
+#if WITH_EDITORONLY_DATA
+	/** Holds all the state editor-only metadata. This array matches in index with the states array. */
+	UPROPERTY()
+	TArray<FSceneStateMetadata> StateMetadata;
+
+	/** Holds all the task editor-only metadata. This array matches in index with the tasks array. */
+	UPROPERTY()
+	TArray<FSceneStateTaskMetadata> TaskMetadata;
+
+	/** Holds all the function editor-only metadata. This array matches in index with the functions. */
+	UPROPERTY()
+	TArray<FSceneStateFunctionMetadata> FunctionMetadata;
+
+	/** Holds all the transition editor-only metadata. This array matches in index with the transitions array. */
+	UPROPERTY()
+	TArray<FSceneStateTransitionMetadata> TransitionMetadata;
+
+	/** Map of a State Node to its State Index */
+	TMap<FObjectKey, uint16> StateNodeToIndex;
+
+	/** Map of a State Machine Graph to its State Machine Index */
+	TMap<FObjectKey, uint16> StateMachineGraphToIndex;
+
+	/** Map of a Task Node to its Task Index */
+	TMap<FObjectKey, uint16> TaskNodeToIndex;
+#endif
+};
+
+template<>
+struct TStructOpsTypeTraits<FSceneStateTemplateData> : public TStructOpsTypeTraitsBase2<FSceneStateTemplateData>
+{
+	enum
+	{
+		WithNoInitConstructor = true
+	};
+};
+
+#undef UE_API
diff --git a/Engine/Plugins/VirtualProduction/SceneState/Source/SceneState/Public/Transition/SceneStateTransition.h b/Engine/Plugins/VirtualProduction/SceneState/Source/SceneState/Public/Transition/SceneStateTransition.h
--- a/Engine/Plugins/VirtualProduction/SceneState/Source/SceneState/Public/Transition/SceneStateTransition.h
+++ b/Engine/Plugins/VirtualProduction/SceneState/Source/SceneState/Public/Transition/SceneStateTransition.h
@@ -13,7 +13,6 @@
 class FBoolProperty;
 class UFunction;
 class UObject;
-class USceneStateGeneratedClass;
 struct FSceneState;
 struct FSceneStateConduit;
 struct FSceneStateExecutionContext;
@@ -42,8 +41,8 @@ struct FSceneStateTransition
 {
 	GENERATED_BODY()
 
-	/** Called when the Generated Class is linking to cache the Event function and Result property */
-	void Link(const FSceneStateTransitionLink& InTransitionLink, USceneStateGeneratedClass* InGeneratedClass);
+	/** Called when the Owner Class is linking to cache the Event function and Result property */
+	void Link(const FSceneStateTransitionLink& InTransitionLink, TNotNull<const UClass*> InOwnerClass);
 
 	/** Called when the State has entered. Allocates and initializes the transition instance */
 	void Setup(const FSceneStateExecutionContext& InContext) const;
diff --git a/Engine/Plugins/VirtualProduction/SceneState/Source/SceneStateBlueprint/Public/SceneStateBlueprint.h b/Engine/Plugins/VirtualProduction/SceneState/Source/SceneStateBlueprint/Public/SceneStateBlueprint.h
--- a/Engine/Plugins/VirtualProduction/SceneState/Source/SceneStateBlueprint/Public/SceneStateBlueprint.h
+++ b/Engine/Plugins/VirtualProduction/SceneState/Source/SceneStateBlueprint/Public/SceneStateBlueprint.h
@@ -10,7 +10,6 @@
 
 #define UE_API SCENESTATEBLUEPRINT_API
 
-class USceneStateGeneratedClass;
 class USceneStateMachineGraph;
 class USceneStateMachineStateNode;
 
diff --git a/Engine/Plugins/VirtualProduction/SceneState/Source/SceneStateBlueprintEditor/Private/Compiler/SceneStateBindingCompiler.cpp b/Engine/Plugins/VirtualProduction/SceneState/Source/SceneStateBlueprintEditor/Private/Compiler/SceneStateBindingCompiler.cpp
--- a/Engine/Plugins/VirtualProduction/SceneState/Source/SceneStateBlueprintEditor/Private/Compiler/SceneStateBindingCompiler.cpp
+++ b/Engine/Plugins/VirtualProduction/SceneState/Source/SceneStateBlueprintEditor/Private/Compiler/SceneStateBindingCompiler.cpp
@@ -8,10 +8,11 @@
 #include "SceneStateBlueprint.h"
 #include "SceneStateBlueprintCompilerContext.h"
 #include "SceneStateEventSchema.h"
-#include "SceneStateGeneratedClass.h"
 #include "SceneStatePropertyReference.h"
+#include "SceneStateTemplateData.h"
 #include "StructUtils/UserDefinedStruct.h"
 #include "Tasks/SceneStateBlueprintableTask.h"
+#include "Tasks/SceneStateTask.h"
 #include "Tasks/SceneStateTaskBindingExtension.h"
 
 #define LOCTEXT_NAMESPACE "SceneStateBindingCompiler"
@@ -59,24 +60,22 @@ EDataAccessType GetBindingBatchingType(const FSceneStateBindingCollection& InBin
 
 } // UE::SceneState::Editor::Private
 
-FBindingCompiler::FBindingCompiler(FBlueprintCompilerContext& InContext, USceneStateBlueprint* InBlueprint, USceneStateGeneratedClass* InGeneratedClass)
+FBindingCompiler::FBindingCompiler(FBlueprintCompilerContext& InContext, TNotNull<USceneStateBlueprint*> InBlueprint, TNotNull<FSceneStateTemplateData*> InTemplateData)
 	: Context(InContext)
 	, Blueprint(InBlueprint)
-	, Class(InGeneratedClass)
-	, BindingFunctionCompiler(InGeneratedClass)
+	, TemplateData(InTemplateData)
+	, BindingFunctionCompiler(InTemplateData)
 {
 }
 
 void FBindingCompiler::Compile()
 {
-	check(Blueprint && Class);
-
 	// Initial estimate count of the binding descs
-	const int32 BindingDescInitialCount = 1 + Class->Tasks.Num() + Class->EventHandlers.Num() + Class->TransitionParameters.Num();
+	const int32 BindingDescInitialCount = 1 + TemplateData->Tasks.Num() + TemplateData->EventHandlers.Num() + TemplateData->TransitionParameters.Num();
 
 	// Copy Editor Bindings and reserve for the binding descs to add
-	Class->BindingCollection.Bindings = Blueprint->BindingCollection.Bindings;
-	Class->BindingCollection.BindingDescs.Empty(BindingDescInitialCount);
+	TemplateData->BindingCollection.Bindings = Blueprint->BindingCollection.Bindings;
+	TemplateData->BindingCollection.BindingDescs.Empty(BindingDescInitialCount);
 
 	ValidBindingMap.Reserve(BindingDescInitialCount);
 
@@ -109,7 +108,7 @@ void FBindingCompiler::AddBindingDesc(FSceneStateBindingDesc&& InBindingDesc)
 	// The binding data view can be filled in with null memory, as only the struct is required for resolution
 	ValidBindingMap.Add(InBindingDesc.ID, FPropertyBindingDataView(InBindingDesc.Struct, nullptr));
 
-	Class->BindingCollection.BindingDescs.Add(MoveTemp(InBindingDesc));
+	TemplateData->BindingCollection.BindingDescs.Add(MoveTemp(InBindingDesc));
 }
 
 void FBindingCompiler::AddRootBindingDesc()
@@ -117,7 +116,7 @@ void FBindingCompiler::AddRootBindingDesc()
 	FSceneStateBindingDesc BindingDesc = Blueprint->CreateRootBinding();
 
 	// Set the Struct to the New Class during compilation
-	BindingDesc.Struct = Class;
+	BindingDesc.Struct = TemplateData->OwnerClass;
 	BindingDesc.DataHandle = FSceneStateBindingDataHandle(ESceneStateDataType::Root);
 
 	AddBindingDesc(MoveTemp(BindingDesc));
@@ -125,11 +124,11 @@ void FBindingCompiler::AddRootBindingDesc()
 
 void FBindingCompiler::AddStateMachineBindingDescs()
 {
-	for (const TPair<FGuid, uint16>& Pair : Class->StateMachineIdToIndex)
+	for (const TPair<FGuid, uint16>& Pair : TemplateData->StateMachineIdToIndex)
 	{
 		const uint16 StateMachineIndex = Pair.Value;
 
-		const FSceneStateMachine& StateMachine = Class->StateMachines[StateMachineIndex];
+		const FSceneStateMachine& StateMachine = TemplateData->StateMachines[StateMachineIndex];
 		if (!StateMachine.Parameters.IsValid())
 		{
 			continue;
@@ -147,11 +146,11 @@ void FBindingCompiler::AddStateMachineBindingDescs()
 
 void FBindingCompiler::AddTransitionBindingDescs()
 {
-	for (const TPair<uint16, FInstancedPropertyBag>& Pair : Class->TransitionParameters)
+	for (const TPair<uint16, FInstancedPropertyBag>& Pair : TemplateData->TransitionParameters)
 	{
 		const uint16 TransitionIndex = Pair.Key;
 
-		const FSceneStateTransitionMetadata& TransitionMetadata = Class->TransitionMetadata[TransitionIndex];
+		const FSceneStateTransitionMetadata& TransitionMetadata = TemplateData->TransitionMetadata[TransitionIndex];
 
 		FSceneStateBindingDesc BindingDesc;
 		BindingDesc.ID = TransitionMetadata.ParametersId;
@@ -165,15 +164,15 @@ void FBindingCompiler::AddTransitionBindingDescs()
 
 void FBindingCompiler::AddTaskBindingDescs()
 {
-	check(Class->Tasks.Num() == Class->TaskMetadata.Num());
+	check(TemplateData->Tasks.Num() == TemplateData->TaskMetadata.Num());
 
-	TArray<FSceneStateBindingDesc>& BindingDescs = Class->BindingCollection.BindingDescs;
+	TArray<FSceneStateBindingDesc>& BindingDescs = TemplateData->BindingCollection.BindingDescs;
 
-	for (int32 TaskIndex = 0; TaskIndex < Class->Tasks.Num(); ++TaskIndex)
+	for (int32 TaskIndex = 0; TaskIndex < TemplateData->Tasks.Num(); ++TaskIndex)
 	{
-		const FSceneStateTask& Task = Class->Tasks[TaskIndex].Get<FSceneStateTask>();
-		const FSceneStateTaskMetadata& TaskMetadata = Class->TaskMetadata[TaskIndex];
-		const FStructView TaskInstance = Class->TaskInstances[TaskIndex];
+		const FSceneStateTask& Task = TemplateData->Tasks[TaskIndex].Get<FSceneStateTask>();
+		const FSceneStateTaskMetadata& TaskMetadata = TemplateData->TaskMetadata[TaskIndex];
+		const FStructView TaskInstance = TemplateData->TaskInstances[TaskIndex];
 
 		// Task Instance Binding
 		{
@@ -215,9 +214,9 @@ void FBindingCompiler::AddTaskBindingDescs()
 
 void FBindingCompiler::AddEventHandlerBindingDescs()
 {
-	for (int32 EventHandlerIndex = 0; EventHandlerIndex < Class->EventHandlers.Num(); ++EventHandlerIndex)
+	for (int32 EventHandlerIndex = 0; EventHandlerIndex < TemplateData->EventHandlers.Num(); ++EventHandlerIndex)
 	{
-		const FSceneStateEventHandler& EventHandler = Class->EventHandlers[EventHandlerIndex];
+		const FSceneStateEventHandler& EventHandler = TemplateData->EventHandlers[EventHandlerIndex];
 
 		const USceneStateEventSchemaObject* EventSchema = EventHandler.GetEventSchemaHandle().GetEventSchema();
 		if (!EventSchema)
@@ -237,13 +236,13 @@ void FBindingCompiler::AddEventHandlerBindingDescs()
 
 void FBindingCompiler::AddFunctionBindingDescs()
 {
-	check(Class->Functions.Num() == Class->FunctionMetadata.Num()
-		&& Class->Functions.Num() == Class->FunctionInstances.Num());
+	check(TemplateData->Functions.Num() == TemplateData->FunctionMetadata.Num()
+		&& TemplateData->Functions.Num() == TemplateData->FunctionInstances.Num());
 
-	for (int32 FunctionIndex = 0; FunctionIndex < Class->Functions.Num(); ++FunctionIndex)
+	for (int32 FunctionIndex = 0; FunctionIndex < TemplateData->Functions.Num(); ++FunctionIndex)
 	{
-		const FStructView FunctionInstance = Class->FunctionInstances[FunctionIndex];
-		const FSceneStateFunctionMetadata& FunctionMetadata = Class->FunctionMetadata[FunctionIndex];
+		const FStructView FunctionInstance = TemplateData->FunctionInstances[FunctionIndex];
+		const FSceneStateFunctionMetadata& FunctionMetadata = TemplateData->FunctionMetadata[FunctionIndex];
 
 		FSceneStateBindingDesc BindingDesc;
 		BindingDesc.ID = FunctionMetadata.FunctionId;
@@ -311,7 +310,7 @@ bool FBindingCompiler::ValidateBinding(const FSceneStateBinding& InBinding) cons
 
 void FBindingCompiler::ResolveBindingDataHandles()
 {
-	for (TArray<FSceneStateBinding>::TIterator Iter(Class->BindingCollection.Bindings); Iter; ++Iter)
+	for (TArray<FSceneStateBinding>::TIterator Iter(TemplateData->BindingCollection.Bindings); Iter; ++Iter)
 	{
 		FSceneStateBinding& Binding = *Iter;
 		Binding.SourceDataHandle = GetDataHandleById(Binding.GetSourcePath().GetStructID());
@@ -328,12 +327,12 @@ void FBindingCompiler::RemoveInvalidBindings()
 {
 	// Remove all the bindings not present in the valid binding map
 	Blueprint->BindingCollection.RemoveInvalidBindings(ValidBindingMap);
-	Class->BindingCollection.RemoveInvalidBindings(ValidBindingMap);
+	TemplateData->BindingCollection.RemoveInvalidBindings(ValidBindingMap);
 }
 
 void FBindingCompiler::GroupBindings()
 {
-	FSceneStateBindingCollection& BindingCollection = Class->BindingCollection;
+	FSceneStateBindingCollection& BindingCollection = TemplateData->BindingCollection;
 
 	// Sort by Batching Type (Copy,Reference) then by the Target Data.
 	BindingCollection.Bindings.StableSort(
@@ -385,23 +384,23 @@ void FBindingCompiler::CompileCopies()
 		return;
 	}
 
-	TConstArrayView<FSceneStateBinding> CopyBindings = MakeArrayView(Class->BindingCollection.Bindings).Slice(CopyRange->Index, CopyRange->Count);
+	TConstArrayView<FSceneStateBinding> CopyBindings = MakeArrayView(TemplateData->BindingCollection.Bindings).Slice(CopyRange->Index, CopyRange->Count);
 
 	// Add 1 batch per group of bindings that all have the same target data handle
 	for (int32 BindingIndex = 0; BindingIndex < CopyBindings.Num();)
 	{
 		const FSceneStateBinding& CopyBinding = CopyBindings[BindingIndex];
 
-		const FSceneStateBindingDesc* TargetDesc = Class->BindingCollection.FindBindingDesc(CopyBinding.TargetDataHandle);
+		const FSceneStateBindingDesc* TargetDesc = TemplateData->BindingCollection.FindBindingDesc(CopyBinding.TargetDataHandle);
 		if (!ensure(TargetDesc))
 		{
 			++BindingIndex;
 			continue;
 		}
 
-		const FPropertyBindingIndex16 BatchBindings(Class->BindingCollection.GetNumCopyBatches());
+		const FPropertyBindingIndex16 BatchBindings(TemplateData->BindingCollection.GetNumCopyBatches());
 
-		FPropertyBindingCopyInfoBatch& CopyBatch = Class->BindingCollection.AddCopyBatch();
+		FPropertyBindingCopyInfoBatch& CopyBatch = TemplateData->BindingCollection.AddCopyBatch();
 		CopyBatch.TargetStruct = TInstancedStruct<FSceneStateBindingDesc>::Make(*TargetDesc);
 		CopyBatch.BindingsBegin = FPropertyBindingIndex16(BindingIndex);
 
@@ -444,9 +443,9 @@ void FBindingCompiler::CompileReferences()
 		return;
 	}
 
-	TConstArrayView<FSceneStateBinding> ReferenceBindings = MakeArrayView(Class->BindingCollection.Bindings).Slice(ReferenceRange->Index, ReferenceRange->Count);
+	TConstArrayView<FSceneStateBinding> ReferenceBindings = MakeArrayView(TemplateData->BindingCollection.Bindings).Slice(ReferenceRange->Index, ReferenceRange->Count);
 
-	Class->BindingCollection.References.Empty(ReferenceBindings.Num());
+	TemplateData->BindingCollection.References.Empty(ReferenceBindings.Num());
 
 	TArray<FPropertyBindingPathIndirection> TargetIndirections;
 	FString ErrorMessage;
@@ -478,9 +477,9 @@ void FBindingCompiler::CompileReferences()
 
 		FSceneStatePropertyReference* PropertyReference = static_cast<FSceneStatePropertyReference*>(const_cast<void*>(TargetIndirections.Last().GetPropertyAddress()));
 		check(PropertyReference);
-		PropertyReference->ReferenceIndex = Class->BindingCollection.References.AddDefaulted();
+		PropertyReference->ReferenceIndex = TemplateData->BindingCollection.References.AddDefaulted();
 
-		FSceneStateBindingReference& Reference =  Class->BindingCollection.References[PropertyReference->ReferenceIndex];
+		FSceneStateBindingReference& Reference =  TemplateData->BindingCollection.References[PropertyReference->ReferenceIndex];
 		Reference.SourcePropertyPath = ReferenceBinding.GetSourcePath();
 		Reference.SourceDataHandle = ReferenceBinding.SourceDataHandle;
 	}
@@ -496,12 +495,12 @@ void FBindingCompiler::OnBindingsBatchCompiled(FPropertyBindingIndex16 InBinding
 
 	if (DataType == ESceneStateDataType::Task)
 	{
-		FSceneStateTask& Task = Class->Tasks[DataIndex].Get<FSceneStateTask>();
+		FSceneStateTask& Task = TemplateData->Tasks[DataIndex].Get<FSceneStateTask>();
 		Task.BindingsBatch = InBindingsBatch;
 	}
 	else if (DataType == ESceneStateDataType::TaskExtension)
 	{
-		FSceneStateTask& Task = Class->Tasks[DataIndex].Get<FSceneStateTask>();
+		FSceneStateTask& Task = TemplateData->Tasks[DataIndex].Get<FSceneStateTask>();
 
 		FSceneStateTaskBindingExtension* BindingExtension = const_cast<FSceneStateTaskBindingExtension*>(Task.GetBindingExtension());
 		check(BindingExtension);
@@ -509,17 +508,17 @@ void FBindingCompiler::OnBindingsBatchCompiled(FPropertyBindingIndex16 InBinding
 	}
 	else if (DataType == ESceneStateDataType::StateMachine)
 	{
-		FSceneStateMachine& StateMachine = Class->StateMachines[DataIndex];
+		FSceneStateMachine& StateMachine = TemplateData->StateMachines[DataIndex];
 		StateMachine.BindingsBatch = InBindingsBatch;
 	}
 	else if (DataType == ESceneStateDataType::Transition)
 	{
-		FSceneStateTransition& Transition = Class->Transitions[DataIndex];
+		FSceneStateTransition& Transition = TemplateData->Transitions[DataIndex];
 		Transition.BindingsBatch = InBindingsBatch;
 	}
 	else if (DataType == ESceneStateDataType::Function)
 	{
-		FSceneStateFunction& Function = Class->Functions[DataIndex].Get<FSceneStateFunction>();
+		FSceneStateFunction& Function = TemplateData->Functions[DataIndex].Get<FSceneStateFunction>();
 		Function.BindingsBatch = InBindingsBatch;
 	}
 	else
@@ -537,9 +536,9 @@ FSceneStateBindingDataHandle FBindingCompiler::GetDataHandleById(const FGuid& In
 	}
 
 	// Tasks
-	for (int32 TaskIndex = 0; TaskIndex < Class->TaskMetadata.Num(); ++TaskIndex)
+	for (int32 TaskIndex = 0; TaskIndex < TemplateData->TaskMetadata.Num(); ++TaskIndex)
 	{
-		const FSceneStateTaskMetadata& TaskMetadata = Class->TaskMetadata[TaskIndex];
+		const FSceneStateTaskMetadata& TaskMetadata = TemplateData->TaskMetadata[TaskIndex];
 		if (TaskMetadata.TaskId == InStructId)
 		{
 			return FSceneStateBindingDataHandle(ESceneStateDataType::Task, TaskIndex);
@@ -547,7 +546,7 @@ FSceneStateBindingDataHandle FBindingCompiler::GetDataHandleById(const FGuid& In
 	}
 
 	// State Machines
-	for (const TPair<FGuid, uint16>& Pair : Class->StateMachineIdToIndex)
+	for (const TPair<FGuid, uint16>& Pair : TemplateData->StateMachineIdToIndex)
 	{
 		if (Pair.Key == InStructId)
 		{
@@ -556,11 +555,11 @@ FSceneStateBindingDataHandle FBindingCompiler::GetDataHandleById(const FGuid& In
 	}
 
 	// Transitions
-	for (const TPair<uint16, FInstancedPropertyBag>& Pair : Class->TransitionParameters)
+	for (const TPair<uint16, FInstancedPropertyBag>& Pair : TemplateData->TransitionParameters)
 	{
 		const uint16 TransitionIndex = Pair.Key;
 
-		const FSceneStateTransitionMetadata& TransitionMetadata = Class->TransitionMetadata[TransitionIndex];
+		const FSceneStateTransitionMetadata& TransitionMetadata = TemplateData->TransitionMetadata[TransitionIndex];
 		if (TransitionMetadata.ParametersId == InStructId)
 		{
 			return FSceneStateBindingDataHandle(ESceneStateDataType::Transition, TransitionIndex);
@@ -568,9 +567,9 @@ FSceneStateBindingDataHandle FBindingCompiler::GetDataHandleById(const FGuid& In
 	}
 
 	// Event Handlers
-	for (int32 EventHandlerIndex = 0; EventHandlerIndex < Class->EventHandlers.Num(); ++EventHandlerIndex)
+	for (int32 EventHandlerIndex = 0; EventHandlerIndex < TemplateData->EventHandlers.Num(); ++EventHandlerIndex)
 	{
-		const FSceneStateEventHandler& EventHandler = Class->EventHandlers[EventHandlerIndex];
+		const FSceneStateEventHandler& EventHandler = TemplateData->EventHandlers[EventHandlerIndex];
 		if (EventHandler.GetHandlerId() == InStructId)
 		{
 			return FSceneStateBindingDataHandle(ESceneStateDataType::EventHandler, EventHandlerIndex);
@@ -578,10 +577,10 @@ FSceneStateBindingDataHandle FBindingCompiler::GetDataHandleById(const FGuid& In
 	}
 
 	// Task Custom Data 
-	for (int32 TaskIndex = 0; TaskIndex < Class->Tasks.Num(); ++TaskIndex)
+	for (int32 TaskIndex = 0; TaskIndex < TemplateData->Tasks.Num(); ++TaskIndex)
 	{
-		const FSceneStateTask& Task = Class->Tasks[TaskIndex].Get<const FSceneStateTask>();
-		const FStructView TaskInstance = Class->TaskInstances[TaskIndex];
+		const FSceneStateTask& Task = TemplateData->Tasks[TaskIndex].Get<const FSceneStateTask>();
+		const FStructView TaskInstance = TemplateData->TaskInstances[TaskIndex];
 
 		FStructView DataView;
 		uint16 DataIndex;
@@ -594,9 +593,9 @@ FSceneStateBindingDataHandle FBindingCompiler::GetDataHandleById(const FGuid& In
 	}
 
 	// Functions
-	for (int32 FunctionIndex = 0; FunctionIndex < Class->Functions.Num(); ++FunctionIndex)
+	for (int32 FunctionIndex = 0; FunctionIndex < TemplateData->Functions.Num(); ++FunctionIndex)
 	{
-		const FSceneStateFunctionMetadata& FunctionMetadata = Class->FunctionMetadata[FunctionIndex];
+		const FSceneStateFunctionMetadata& FunctionMetadata = TemplateData->FunctionMetadata[FunctionIndex];
 		if (FunctionMetadata.FunctionId == InStructId)
 		{
 			return FSceneStateBindingDataHandle(ESceneStateDataType::Function, FunctionIndex);
diff --git a/Engine/Plugins/VirtualProduction/SceneState/Source/SceneStateBlueprintEditor/Private/Compiler/SceneStateBindingCompiler.h b/Engine/Plugins/VirtualProduction/SceneState/Source/SceneStateBlueprintEditor/Private/Compiler/SceneStateBindingCompiler.h
--- a/Engine/Plugins/VirtualProduction/SceneState/Source/SceneStateBlueprintEditor/Private/Compiler/SceneStateBindingCompiler.h
+++ b/Engine/Plugins/VirtualProduction/SceneState/Source/SceneStateBlueprintEditor/Private/Compiler/SceneStateBindingCompiler.h
@@ -9,7 +9,6 @@
 #include "UObject/ObjectPtr.h"
 
 class USceneStateBlueprint;
-class USceneStateGeneratedClass;
 struct FPropertyBindingIndex16;
 struct FSceneStateBinding;
 struct FSceneStateBindingDesc;
@@ -33,7 +32,7 @@ enum class EDataAccessType : uint8
 class FBindingCompiler
 {
 public:
-	explicit FBindingCompiler(FBlueprintCompilerContext& InContext, USceneStateBlueprint* InBlueprint, USceneStateGeneratedClass* InGeneratedClass);
+	explicit FBindingCompiler(FBlueprintCompilerContext& InContext, TNotNull<USceneStateBlueprint*> InBlueprint, TNotNull<FSceneStateTemplateData*> InTemplateData);
 
 	void Compile();
 
@@ -67,8 +66,8 @@ private:
 	/** The blueprint that generated the class */
 	TObjectPtr<USceneStateBlueprint> Blueprint;
 
-	/** The generated class to compile bindings for */
-	TObjectPtr<USceneStateGeneratedClass> Class;
+	/** The target to store the compiled bindings */
+	FSceneStateTemplateData* TemplateData;
 
 	/** Compiler for the functions found within the bindings */
 	FBindingFunctionCompiler BindingFunctionCompiler;
diff --git a/Engine/Plugins/VirtualProduction/SceneState/Source/SceneStateBlueprintEditor/Private/Compiler/SceneStateBindingFunctionCompiler.cpp b/Engine/Plugins/VirtualProduction/SceneState/Source/SceneStateBlueprintEditor/Private/Compiler/SceneStateBindingFunctionCompiler.cpp
--- a/Engine/Plugins/VirtualProduction/SceneState/Source/SceneStateBlueprintEditor/Private/Compiler/SceneStateBindingFunctionCompiler.cpp
+++ b/Engine/Plugins/VirtualProduction/SceneState/Source/SceneStateBlueprintEditor/Private/Compiler/SceneStateBindingFunctionCompiler.cpp
@@ -2,7 +2,8 @@
 
 #include "SceneStateBindingFunctionCompiler.h"
 #include "SceneStateBindingFunction.h"
-#include "SceneStateGeneratedClass.h"
+#include "SceneStateTemplateData.h"
+#include "Tasks/SceneStateTask.h"
 #include "Tasks/SceneStateTaskBindingExtension.h"
 
 namespace UE::SceneState::Editor
@@ -36,7 +37,7 @@ void FBindingFunctionCompiler::CompileFunctions(const FGuid& InTargetStructId)
 	}
 	FunctionRangeMap.Add(InTargetStructId);
 
-	const FSceneStateBindingCollection& BindingCollection = Class->GetBindingCollection();
+	const FSceneStateBindingCollection& BindingCollection = TemplateData->BindingCollection;
 
 	TArray<FConstStructView> Functions;
 	TArray<FConstStructView> FunctionInstances;
@@ -73,17 +74,17 @@ void FBindingFunctionCompiler::CompileFunctions(const FGuid& InTargetStructId)
 	}
 
 	FSceneStateRange& FunctionRange = FunctionRangeMap[InTargetStructId];
-	FunctionRange.Index = Class->Functions.Num();
+	FunctionRange.Index = TemplateData->Functions.Num();
 	FunctionRange.Count = Functions.Num();
 
-	Class->Functions.Append(Functions);
-	Class->FunctionInstances.Append(FunctionInstances);
-	Class->FunctionMetadata.Append(MoveTemp(FunctionMetadata));
+	TemplateData->Functions.Append(Functions);
+	TemplateData->FunctionInstances.Append(FunctionInstances);
+	TemplateData->FunctionMetadata.Append(MoveTemp(FunctionMetadata));
 }
 
 void FBindingFunctionCompiler::CompileStateMachineFunctions()
 {
-	for (const TPair<FGuid, uint16>& Pair : Class->StateMachineIdToIndex)
+	for (const TPair<FGuid, uint16>& Pair : TemplateData->StateMachineIdToIndex)
 	{
 		CompileFunctions(Pair.Key);
 	}
@@ -91,24 +92,24 @@ void FBindingFunctionCompiler::CompileStateMachineFunctions()
 
 void FBindingFunctionCompiler::CompileTransitionFunctions()
 {
-	for (const TPair<uint16, FInstancedPropertyBag>& Pair : Class->TransitionParameters)
+	for (const TPair<uint16, FInstancedPropertyBag>& Pair : TemplateData->TransitionParameters)
 	{
-		const FSceneStateTransitionMetadata& TransitionMetadata = Class->TransitionMetadata[Pair.Key];
+		const FSceneStateTransitionMetadata& TransitionMetadata = TemplateData->TransitionMetadata[Pair.Key];
 		CompileFunctions(TransitionMetadata.ParametersId);
 	}
 }
 
 void FBindingFunctionCompiler::CompileTaskFunctions()
 {
-	check(Class->Tasks.Num() == Class->GetTaskMetadata().Num());
+	check(TemplateData->Tasks.Num() == TemplateData->TaskMetadata.Num());
 
-	for (int32 TaskIndex = 0; TaskIndex < Class->Tasks.Num(); ++TaskIndex)
+	for (int32 TaskIndex = 0; TaskIndex < TemplateData->Tasks.Num(); ++TaskIndex)
 	{
-		const FSceneStateTaskMetadata& TaskMetadata = Class->TaskMetadata[TaskIndex];
+		const FSceneStateTaskMetadata& TaskMetadata = TemplateData->TaskMetadata[TaskIndex];
 		CompileFunctions(TaskMetadata.TaskId);
 
-		const FSceneStateTask& Task = Class->Tasks[TaskIndex].Get<const FSceneStateTask>();
-		const FStructView TaskInstance = Class->TaskInstances[TaskIndex];
+		const FSceneStateTask& Task = TemplateData->Tasks[TaskIndex].Get<const FSceneStateTask>();
+		const FStructView TaskInstance = TemplateData->TaskInstances[TaskIndex];
 
 		if (const FSceneStateTaskBindingExtension* BindingExtension = Task.GetBindingExtension())
 		{
diff --git a/Engine/Plugins/VirtualProduction/SceneState/Source/SceneStateBlueprintEditor/Private/Compiler/SceneStateBindingFunctionCompiler.h b/Engine/Plugins/VirtualProduction/SceneState/Source/SceneStateBlueprintEditor/Private/Compiler/SceneStateBindingFunctionCompiler.h
--- a/Engine/Plugins/VirtualProduction/SceneState/Source/SceneStateBlueprintEditor/Private/Compiler/SceneStateBindingFunctionCompiler.h
+++ b/Engine/Plugins/VirtualProduction/SceneState/Source/SceneStateBlueprintEditor/Private/Compiler/SceneStateBindingFunctionCompiler.h
@@ -5,7 +5,7 @@
 #include "SceneStateRange.h"
 #include "UObject/ObjectPtr.h"
 
-class USceneStateGeneratedClass;
+struct FSceneStateTemplateData;
 
 namespace UE::SceneState::Editor
 {
@@ -14,8 +14,8 @@ namespace UE::SceneState::Editor
 class FBindingFunctionCompiler
 {
 public:
-	explicit FBindingFunctionCompiler(USceneStateGeneratedClass* InGeneratedClass)
-		: Class(InGeneratedClass)
+	explicit FBindingFunctionCompiler(TNotNull<FSceneStateTemplateData*> InTemplateData)
+		: TemplateData(InTemplateData)
 	{
 	}
 
@@ -38,8 +38,8 @@ private:
 	/** Compiles the functions for all the compiled transitions */
 	void CompileTransitionFunctions();
 
-	/** The generated class to compile functions for */
-	TObjectPtr<USceneStateGeneratedClass> Class;
+	/** The target to store the compiled functions */
+	FSceneStateTemplateData* TemplateData;
 
 	/** Map of the target struct id to its range of the functions in the generated class */
 	TMap<FGuid, FSceneStateRange> FunctionRangeMap;
diff --git a/Engine/Plugins/VirtualProduction/SceneState/Source/SceneStateBlueprintEditor/Private/Compiler/SceneStateBlueprintCompilerContext.cpp b/Engine/Plugins/VirtualProduction/SceneState/Source/SceneStateBlueprintEditor/Private/Compiler/SceneStateBlueprintCompilerContext.cpp
--- a/Engine/Plugins/VirtualProduction/SceneState/Source/SceneStateBlueprintEditor/Private/Compiler/SceneStateBlueprintCompilerContext.cpp
+++ b/Engine/Plugins/VirtualProduction/SceneState/Source/SceneStateBlueprintEditor/Private/Compiler/SceneStateBlueprintCompilerContext.cpp
@@ -8,8 +8,10 @@
 #include "SceneStateMachine.h"
 #include "SceneStateMachineCompiler.h"
 #include "SceneStateMachineGraph.h"
+#include "SceneStateTemplateData.h"
 #include "SceneStateTransitionGraphCompiler.h"
 #include "SceneStateUtils.h"
+#include "Tasks/SceneStateTask.h"
 
 namespace UE::SceneState::Editor
 {
@@ -113,6 +115,8 @@ void FBlueprintCompilerContext::MergeUbergraphPagesIn(UEdGraph* InUbergraph)
 
 	USceneStateBlueprint* SceneStateBlueprint = CastChecked<USceneStateBlueprint>(Blueprint);
 
+	FSceneStateTemplateData& TemplateData = NewGeneratedClass->SharedTemplateData.Get<FSceneStateTemplateData>();
+
 	// Top Level State Machines
 	TArray<FSceneStateMachine> StateMachines;
 	StateMachines.Reserve(SceneStateBlueprint->StateMachineGraphs.Num());
@@ -123,8 +127,8 @@ void FBlueprintCompilerContext::MergeUbergraphPagesIn(UEdGraph* InUbergraph)
 	TMap<FObjectKey, uint16> StateMachineGraphToIndex;
 	StateMachineGraphToIndex.Reserve(SceneStateBlueprint->StateMachineGraphs.Num());
 
-	const int32 RootStateIndex = NewGeneratedClass->States.AddDefaulted();
-	const int32 RootStateMetadataIndex = NewGeneratedClass->StateMetadata.AddDefaulted();
+	const int32 RootStateIndex = TemplateData.States.AddDefaulted();
+	const int32 RootStateMetadataIndex = TemplateData.StateMetadata.AddDefaulted();
 
 	check(RootStateIndex == RootStateMetadataIndex);
 
@@ -157,22 +161,22 @@ void FBlueprintCompilerContext::MergeUbergraphPagesIn(UEdGraph* InUbergraph)
 		}
 	}
 
-	NewGeneratedClass->RootStateIndex = RootStateIndex;
+	TemplateData.RootStateIndex = RootStateIndex;
 
 	// Root State only considers the range of the Auto-Run State Machines
-	FSceneState& RootState = NewGeneratedClass->States[RootStateIndex];
-	RootState.StateMachineRange.Index = NewGeneratedClass->StateMachines.Num();
+	FSceneState& RootState = TemplateData.States[RootStateIndex];
+	RootState.StateMachineRange.Index = TemplateData.StateMachines.Num();
 	RootState.StateMachineRange.Count = AutoRunCount;
 
-	NewGeneratedClass->StateMachines.Append(MoveTemp(StateMachines));
+	TemplateData.StateMachines.Append(MoveTemp(StateMachines));
 
 	// Upgrade the map to absolute indices before baking it to the generated class' map
 	ToAbsoluteIndexMap(StateMachineIdToIndex, RootState.StateMachineRange.Index);
-	NewGeneratedClass->StateMachineIdToIndex.Append(MoveTemp(StateMachineIdToIndex));
+	TemplateData.StateMachineIdToIndex.Append(MoveTemp(StateMachineIdToIndex));
 
 	// Upgrade the map to absolute indices before baking it to the generated class' map
 	ToAbsoluteIndexMap(StateMachineGraphToIndex, RootState.StateMachineRange.Index);
-	NewGeneratedClass->StateMachineGraphToIndex.Append(MoveTemp(StateMachineGraphToIndex));
+	TemplateData.StateMachineGraphToIndex.Append(MoveTemp(StateMachineGraphToIndex));
 
 	UpdateTaskIndices();
 	CompileBindings();
@@ -189,7 +193,14 @@ UBlueprint* FBlueprintCompilerContext::GetBlueprint() const
 	return Blueprint;
 }
 
-USceneStateGeneratedClass* FBlueprintCompilerContext::GetGeneratedClass() const
+FSceneStateTemplateData* FBlueprintCompilerContext::GetTemplateData() const
+{
+	return NewGeneratedClass
+		? NewGeneratedClass->SharedTemplateData.GetPtr<FSceneStateTemplateData>()
+		: nullptr;
+}
+
+UObject* FBlueprintCompilerContext::GetTargetOuter() const
 {
 	return NewGeneratedClass;
 }
@@ -214,9 +225,10 @@ FTransitionGraphCompileResult FBlueprintCompilerContext::CompileTransitionGraph(
 
 void FBlueprintCompilerContext::UpdateTaskIndices()
 {
-	TConstArrayView<FSceneState> States = NewGeneratedClass->GetStates();
+	FSceneStateTemplateData& TemplateData = NewGeneratedClass->SharedTemplateData.Get<FSceneStateTemplateData>();
 
-	FInstancedStructContainer& Tasks = NewGeneratedClass->Tasks;
+	TConstArrayView<FSceneState> States = TemplateData.States;
+	FInstancedStructContainer& Tasks = TemplateData.Tasks;
 
 	for (int32 StateIndex = 0; StateIndex < States.Num(); ++StateIndex)
 	{
@@ -237,7 +249,9 @@ void FBlueprintCompilerContext::UpdateTaskIndices()
 
 void FBlueprintCompilerContext::CompileBindings()
 {
-	FBindingCompiler BindingCompiler(*this, CastChecked<USceneStateBlueprint>(Blueprint), NewGeneratedClass);
+	FSceneStateTemplateData& TemplateData = NewGeneratedClass->SharedTemplateData.Get<FSceneStateTemplateData>();
+
+	FBindingCompiler BindingCompiler(*this, CastChecked<USceneStateBlueprint>(Blueprint), &TemplateData);
 	BindingCompiler.Compile();
 }
 
diff --git a/Engine/Plugins/VirtualProduction/SceneState/Source/SceneStateBlueprintEditor/Private/Compiler/SceneStateBlueprintCompilerContext.h b/Engine/Plugins/VirtualProduction/SceneState/Source/SceneStateBlueprintEditor/Private/Compiler/SceneStateBlueprintCompilerContext.h
--- a/Engine/Plugins/VirtualProduction/SceneState/Source/SceneStateBlueprintEditor/Private/Compiler/SceneStateBlueprintCompilerContext.h
+++ b/Engine/Plugins/VirtualProduction/SceneState/Source/SceneStateBlueprintEditor/Private/Compiler/SceneStateBlueprintCompilerContext.h
@@ -29,7 +29,8 @@ private:
 
 	//~ Begin ISceneStateMachineCompilerContext
 	virtual UBlueprint* GetBlueprint() const override;
-	virtual USceneStateGeneratedClass* GetGeneratedClass() const override;
+	virtual FSceneStateTemplateData* GetTemplateData() const override;
+	virtual UObject* GetTargetOuter() const override;
 	virtual FTransitionGraphCompileResult CompileTransitionGraph(USceneStateTransitionGraph* InTransitionGraph) override;
 	//~ End ISceneStateMachineCompilerContext
 
diff --git a/Engine/Plugins/VirtualProduction/SceneState/Source/SceneStateBlueprintEditor/Private/DebugExecutors/SceneStateDebugExecutor.cpp b/Engine/Plugins/VirtualProduction/SceneState/Source/SceneStateBlueprintEditor/Private/DebugExecutors/SceneStateDebugExecutor.cpp
--- a/Engine/Plugins/VirtualProduction/SceneState/Source/SceneStateBlueprintEditor/Private/DebugExecutors/SceneStateDebugExecutor.cpp
+++ b/Engine/Plugins/VirtualProduction/SceneState/Source/SceneStateBlueprintEditor/Private/DebugExecutors/SceneStateDebugExecutor.cpp
@@ -4,7 +4,6 @@
 #include "Kismet2/BlueprintEditorUtils.h"
 #include "Nodes/SceneStateMachineNode.h"
 #include "SceneStateBlueprint.h"
-#include "SceneStateGeneratedClass.h"
 
 namespace UE::SceneState::Editor
 {
diff --git a/Engine/Plugins/VirtualProduction/SceneState/Source/SceneStateBlueprintEditor/Private/DebugExecutors/SceneStateDebugStateExecutor.cpp b/Engine/Plugins/VirtualProduction/SceneState/Source/SceneStateBlueprintEditor/Private/DebugExecutors/SceneStateDebugStateExecutor.cpp
--- a/Engine/Plugins/VirtualProduction/SceneState/Source/SceneStateBlueprintEditor/Private/DebugExecutors/SceneStateDebugStateExecutor.cpp
+++ b/Engine/Plugins/VirtualProduction/SceneState/Source/SceneStateBlueprintEditor/Private/DebugExecutors/SceneStateDebugStateExecutor.cpp
@@ -1,18 +1,14 @@
 // Copyright Epic Games, Inc. All Rights Reserved.
 
 #include "SceneStateDebugStateExecutor.h"
-#include "SceneStateGeneratedClass.h"
+#include "SceneStateTemplateData.h"
 
 namespace UE::SceneState::Editor
 {
 
 const FSceneState* FDebugStateExecutor::GetState(const FSceneStateExecutionContext& InExecutionContext) const
 {
-	if (const USceneStateGeneratedClass* GeneratedClass = InExecutionContext.GetGeneratedClass())
-	{
-		return GeneratedClass->FindStateFromNode(GetNodeKey());
-	}
-	return nullptr;
+	return InExecutionContext.GetSharedTemplateData().FindStateFromNode(GetNodeKey());
 }
 
 void FDebugStateExecutor::OnStart(const FSceneStateExecutionContext& InExecutionContext)
diff --git a/Engine/Plugins/VirtualProduction/SceneState/Source/SceneStateBlueprintEditor/Private/DebugExecutors/SceneStateDebugTaskExecutor.cpp b/Engine/Plugins/VirtualProduction/SceneState/Source/SceneStateBlueprintEditor/Private/DebugExecutors/SceneStateDebugTaskExecutor.cpp
--- a/Engine/Plugins/VirtualProduction/SceneState/Source/SceneStateBlueprintEditor/Private/DebugExecutors/SceneStateDebugTaskExecutor.cpp
+++ b/Engine/Plugins/VirtualProduction/SceneState/Source/SceneStateBlueprintEditor/Private/DebugExecutors/SceneStateDebugTaskExecutor.cpp
@@ -3,8 +3,8 @@
 #include "SceneStateDebugTaskExecutor.h"
 #include "Nodes/SceneStateMachineStateNode.h"
 #include "Nodes/SceneStateMachineTaskNode.h"
-#include "SceneStateGeneratedClass.h"
 #include "SceneStateMachineGraphSchema.h"
+#include "SceneStateTemplateData.h"
 #include "Tasks/SceneStateTaskInstance.h"
 
 namespace UE::SceneState::Editor
@@ -12,20 +12,16 @@ namespace UE::SceneState::Editor
 
 const FSceneStateTask* FDebugTaskExecutor::GetTask(const FSceneStateExecutionContext& InExecutionContext) const
 {
-	if (const USceneStateGeneratedClass* GeneratedClass = InExecutionContext.GetGeneratedClass())
-	{
-		return GeneratedClass->FindTaskFromNode(GetNodeKey());
-	}
-	return nullptr;
+	return InExecutionContext.GetSharedTemplateData().FindTaskFromNode(GetNodeKey());
 }
 
-const FSceneState* FDebugTaskExecutor::GetState(const USceneStateGeneratedClass* InGeneratedClass) const
+const FSceneState* FDebugTaskExecutor::GetState(const FSceneStateTemplateData& InTemplateData) const
 {
 	if (const USceneStateMachineTaskNode* const TaskNode = Cast<USceneStateMachineTaskNode>(GetNodeKey().ResolveObjectPtr()))
 	{
 		if (const USceneStateMachineStateNode* const StateNode = USceneStateMachineGraphSchema::FindConnectedStateNode(TaskNode))
 		{
-			return InGeneratedClass->FindStateFromNode(StateNode);
+			return InTemplateData.FindStateFromNode(StateNode);
 		}
 	}
 	return nullptr;
@@ -33,13 +29,9 @@ const FSceneState* FDebugTaskExecutor::GetState(const USceneStateGeneratedClass*
 
 void FDebugTaskExecutor::Setup(const FSceneStateExecutionContext& InExecutionContext, const FSceneStateTask& InTask)
 {
-	const USceneStateGeneratedClass* GeneratedClass = InExecutionContext.GetGeneratedClass();
-	if (!GeneratedClass)
-	{
-		return;
-	}
+	const FSceneStateTemplateData& TemplateData = InExecutionContext.GetSharedTemplateData();
 
-	const FSceneState* ParentState = GetState(GeneratedClass);
+	const FSceneState* ParentState = GetState(TemplateData);
 	if (!ParentState)
 	{
 		return;
@@ -54,7 +46,7 @@ void FDebugTaskExecutor::Setup(const FSceneStateExecutionContext& InExecutionCon
 	// and only set the task instance template for the task this executor cares about
 	TArray<FConstStructView> TemplateTaskInstances;
 	TemplateTaskInstances.SetNum(TaskRange.Count);
-	TemplateTaskInstances[TaskRelativeIndex] = GeneratedClass->GetTaskInstances()[TaskAbsoluteIndex];
+	TemplateTaskInstances[TaskRelativeIndex] = TemplateData.TaskInstances[TaskAbsoluteIndex];
 
 	ParentState->AllocateTaskInstances(InExecutionContext, TemplateTaskInstances);
 }
diff --git a/Engine/Plugins/VirtualProduction/SceneState/Source/SceneStateBlueprintEditor/Private/DebugExecutors/SceneStateDebugTaskExecutor.h b/Engine/Plugins/VirtualProduction/SceneState/Source/SceneStateBlueprintEditor/Private/DebugExecutors/SceneStateDebugTaskExecutor.h
--- a/Engine/Plugins/VirtualProduction/SceneState/Source/SceneStateBlueprintEditor/Private/DebugExecutors/SceneStateDebugTaskExecutor.h
+++ b/Engine/Plugins/VirtualProduction/SceneState/Source/SceneStateBlueprintEditor/Private/DebugExecutors/SceneStateDebugTaskExecutor.h
@@ -21,7 +21,7 @@ protected:
 	const FSceneStateTask* GetTask(const FSceneStateExecutionContext& InExecutionContext) const;
 
 	/** Retrieves the state corresponding to the stored state node */
-	const FSceneState* GetState(const USceneStateGeneratedClass* InGeneratedClass) const;
+	const FSceneState* GetState(const FSceneStateTemplateData& InTemplateData) const;
 
 	/** Called to allocate the task instance */
 	void Setup(const FSceneStateExecutionContext& InExecutionContext, const FSceneStateTask& InTask);
diff --git a/Engine/Plugins/VirtualProduction/SceneState/Source/SceneStateMachineEditor/Private/SceneStateMachineCompiler.cpp b/Engine/Plugins/VirtualProduction/SceneState/Source/SceneStateMachineEditor/Private/SceneStateMachineCompiler.cpp
--- a/Engine/Plugins/VirtualProduction/SceneState/Source/SceneStateMachineEditor/Private/SceneStateMachineCompiler.cpp
+++ b/Engine/Plugins/VirtualProduction/SceneState/Source/SceneStateMachineEditor/Private/SceneStateMachineCompiler.cpp
@@ -13,15 +13,16 @@
 #include "Nodes/SceneStateMachineTransitionNode.h"
 #include "Nodes/SceneStateTransitionResultNode.h"
 #include "SceneState.h"
-#include "SceneStateGeneratedClass.h"
 #include "SceneStateMachine.h"
 #include "SceneStateMachineConduitCompiler.h"
 #include "SceneStateMachineGraph.h"
 #include "SceneStateMachineTaskCompiler.h"
 #include "SceneStateMachineTransitionCompiler.h"
+#include "SceneStateTemplateData.h"
 #include "SceneStateTransitionGraph.h"
 #include "SceneStateUtils.h"
 #include "Tasks/SceneStateBlueprintableTask.h"
+#include "Tasks/SceneStateTask.h"
 #include "UObject/UObjectThreadContext.h"
 
 namespace UE::SceneState::Editor
@@ -56,18 +57,17 @@ TMap<uint16, FInstancedPropertyBag> ConvertParameterMap(FSceneStateRange InTrans
 } // Private
 
 FStateMachineCompiler::FStateMachineCompiler(USceneStateMachineGraph* InGraph, IStateMachineCompilerContext& InCompilerContext)
-	: StateMachineGraph(InGraph)
+	: Outer(InCompilerContext.GetTargetOuter())
+	, TemplateData(InCompilerContext.GetTemplateData())
+	, StateMachineGraph(InGraph)
+	, Blueprint(FBlueprintEditorUtils::FindBlueprintForGraph(InGraph))
 	, Context(InCompilerContext)
 {
-	GeneratedClass = InCompilerContext.GetGeneratedClass();
-	check(StateMachineGraph && GeneratedClass);
-	Blueprint = FBlueprintEditorUtils::FindBlueprintForGraph(InGraph);
+	check(TemplateData && StateMachineGraph && Outer);
 }
 
 FSceneStateMachine FStateMachineCompiler::Compile()
 {
-	check(StateMachineGraph && GeneratedClass);
-
 	FSceneStateMachine StateMachine;
 
 	USceneStateMachineEntryNode* EntryNode = StateMachineGraph->GetEntryNode();
@@ -114,24 +114,24 @@ FSceneStateMachine FStateMachineCompiler::Compile()
 		CompileStateTransitions(SceneState, StateNode);
 	}
 
-	StateMachine.StateRange.Index = GeneratedClass->States.Num();
+	StateMachine.StateRange.Index = TemplateData->States.Num();
 	StateMachine.StateRange.Count = States.Num();
 
-	GeneratedClass->States.Append(MoveTemp(States));
-	GeneratedClass->StateMetadata.Append(MoveTemp(StateMetadata));
+	TemplateData->States.Append(MoveTemp(States));
+	TemplateData->StateMetadata.Append(MoveTemp(StateMetadata));
 
 	States.Reset();
 	StateMetadata.Reset();
 
 	// Upgrade the map to absolute indices before baking it to the generated class' map
 	ToAbsoluteIndexMap(StateNodeIndexMap, StateMachine.StateRange.Index);
-	GeneratedClass->StateNodeToIndex.Append(MoveTemp(StateNodeIndexMap));
+	TemplateData->StateNodeToIndex.Append(MoveTemp(StateNodeIndexMap));
 
-	StateMachine.ConduitRange.Index = GeneratedClass->Conduits.Num();
+	StateMachine.ConduitRange.Index = TemplateData->Conduits.Num();
 	StateMachine.ConduitRange.Count = Conduits.Num();
 
-	GeneratedClass->Conduits.Append(MoveTemp(Conduits));
-	GeneratedClass->ConduitLinks.Append(MoveTemp(ConduitLinks));
+	TemplateData->Conduits.Append(MoveTemp(Conduits));
+	TemplateData->ConduitLinks.Append(MoveTemp(ConduitLinks));
 
 	Conduits.Reset();
 	ConduitLinks.Reset();
@@ -221,13 +221,13 @@ void FStateMachineCompiler::CompileSubStateMachines(FSceneState& InNewState, con
 		}
 	}
 
-	InNewState.StateMachineRange.Index = GeneratedClass->StateMachines.Num();
+	InNewState.StateMachineRange.Index = TemplateData->StateMachines.Num();
 	InNewState.StateMachineRange.Count = SubStateMachines.Num();
-	GeneratedClass->StateMachines.Append(MoveTemp(SubStateMachines));
+	TemplateData->StateMachines.Append(MoveTemp(SubStateMachines));
 
 	// Upgrade the map to absolute indices before baking it to the generated class' map
 	ToAbsoluteIndexMap(StateMachineGraphToIndex, InNewState.StateMachineRange.Index);
-	GeneratedClass->StateMachineGraphToIndex.Append(MoveTemp(StateMachineGraphToIndex));
+	TemplateData->StateMachineGraphToIndex.Append(MoveTemp(StateMachineGraphToIndex));
 }
 
 void FStateMachineCompiler::AddConduitToCompile(const USceneStateMachineConduitNode* InConduitNode)
@@ -262,28 +262,28 @@ void FStateMachineCompiler::CompileTasks(FSceneState& InNewState, const USceneSt
 {
 	FStateMachineTaskCompiler::FCompileResult TaskCompileResult;
 	{
-		FStateMachineTaskCompiler TaskCompiler(InStateNode->GetTaskPin(), /*Outer*/GeneratedClass);
+		FStateMachineTaskCompiler TaskCompiler(InStateNode->GetTaskPin(), /*Outer*/Outer);
 		TaskCompiler.Compile(TaskCompileResult);
 	}
 
-	InNewState.TaskRange.Index = GeneratedClass->Tasks.Num();
+	InNewState.TaskRange.Index = TemplateData->Tasks.Num();
 	InNewState.TaskRange.Count = TaskCompileResult.Tasks.Num();
 
 	// Convert the prerequisite range from relative to absolute
-	const uint16 TaskPrerequisiteIndex = GeneratedClass->TaskPrerequisites.Num();
+	const uint16 TaskPrerequisiteIndex = TemplateData->TaskPrerequisites.Num();
 	for (FInstancedStruct& Task : TaskCompileResult.Tasks)
 	{
 		Task.GetMutable<FSceneStateTask>().PrerequisiteRange.Index += TaskPrerequisiteIndex;
 	}
 
 	// todo: bigger chunks to reduce FInstancedStructContainer realloc
-	GeneratedClass->Tasks.Append(MoveTemp(TaskCompileResult.Tasks));
-	GeneratedClass->TaskInstances.Append(TaskCompileResult.TaskInstances);
-	GeneratedClass->TaskMetadata.Append(MoveTemp(TaskCompileResult.TaskMetadata));
-	GeneratedClass->TaskPrerequisites.Append(MoveTemp(TaskCompileResult.TaskPrerequisites));
+	TemplateData->Tasks.Append(MoveTemp(TaskCompileResult.Tasks));
+	TemplateData->TaskInstances.Append(TaskCompileResult.TaskInstances);
+	TemplateData->TaskMetadata.Append(MoveTemp(TaskCompileResult.TaskMetadata));
+	TemplateData->TaskPrerequisites.Append(MoveTemp(TaskCompileResult.TaskPrerequisites));
 
-	InNewState.InstanceTaskObjects(GeneratedClass
-		, GetStructViews(GeneratedClass->TaskInstances, InNewState.TaskRange)
+	InNewState.InstanceTaskObjects(Outer
+		, GetStructViews(TemplateData->TaskInstances, InNewState.TaskRange)
 		, TaskCompileResult.TaskInstances
 		, [this](FObjectDuplicationParameters& InParams)->UObject*
 		{
@@ -291,26 +291,26 @@ void FStateMachineCompiler::CompileTasks(FSceneState& InNewState, const USceneSt
 		});
 
 	ToAbsoluteIndexMap(TaskCompileResult.TaskToIndexMap, InNewState.TaskRange.Index);
-	GeneratedClass->TaskNodeToIndex.Append(MoveTemp(TaskCompileResult.TaskToIndexMap));
+	TemplateData->TaskNodeToIndex.Append(MoveTemp(TaskCompileResult.TaskToIndexMap));
 }
 
 void FStateMachineCompiler::CompileEventHandlers(FSceneState& InNewState, const USceneStateMachineStateNode* InStateNode)
 {
-	InNewState.EventHandlerRange.Index = GeneratedClass->EventHandlers.Num();
+	InNewState.EventHandlerRange.Index = TemplateData->EventHandlers.Num();
 	InNewState.EventHandlerRange.Count = InStateNode->EventHandlers.Num();
 
-	GeneratedClass->EventHandlers.Append(InStateNode->EventHandlers);
+	TemplateData->EventHandlers.Append(InStateNode->EventHandlers);
 }
 
 void FStateMachineCompiler::FinishTransitionCompilation(FSceneStateRange& OutTransitionRange, FStateMachineTransitionCompileResult&& InCompileResult)
 {
-	OutTransitionRange.Index = GeneratedClass->Transitions.Num();
+	OutTransitionRange.Index = TemplateData->Transitions.Num();
 	OutTransitionRange.Count = InCompileResult.Transitions.Num();
 
-	GeneratedClass->Transitions.Append(MoveTemp(InCompileResult.Transitions));
-	GeneratedClass->TransitionLinks.Append(MoveTemp(InCompileResult.Links));
-	GeneratedClass->TransitionMetadata.Append(MoveTemp(InCompileResult.Metadata));
-	GeneratedClass->TransitionParameters.Append(Private::ConvertParameterMap(OutTransitionRange, InCompileResult.Parameters));
+	TemplateData->Transitions.Append(MoveTemp(InCompileResult.Transitions));
+	TemplateData->TransitionLinks.Append(MoveTemp(InCompileResult.Links));
+	TemplateData->TransitionMetadata.Append(MoveTemp(InCompileResult.Metadata));
+	TemplateData->TransitionParameters.Append(Private::ConvertParameterMap(OutTransitionRange, InCompileResult.Parameters));
 }
 
 void FStateMachineCompiler::CompileStateTransitions(FSceneState& InNewState, const USceneStateMachineStateNode* InStateNode)
diff --git a/Engine/Plugins/VirtualProduction/SceneState/Source/SceneStateMachineEditor/Private/Widgets/SSceneStateMachineStateNode.cpp b/Engine/Plugins/VirtualProduction/SceneState/Source/SceneStateMachineEditor/Private/Widgets/SSceneStateMachineStateNode.cpp
--- a/Engine/Plugins/VirtualProduction/SceneState/Source/SceneStateMachineEditor/Private/Widgets/SSceneStateMachineStateNode.cpp
+++ b/Engine/Plugins/VirtualProduction/SceneState/Source/SceneStateMachineEditor/Private/Widgets/SSceneStateMachineStateNode.cpp
@@ -7,12 +7,12 @@
 #include "SGraphPin.h"
 #include "SGraphPreviewer.h"
 #include "SSceneStateMachineOutputPin.h"
-#include "SceneStateGeneratedClass.h"
 #include "SceneStateMachine.h"
 #include "SceneStateMachineEditorStyle.h"
 #include "SceneStateMachineGraphSchema.h"
 #include "SceneStateMachineInstance.h"
 #include "SceneStateObject.h"
+#include "SceneStateTemplateData.h"
 #include "Widgets/Images/SImage.h"
 #include "Widgets/Layout/SBorder.h"
 #include "Widgets/Notifications/SErrorText.h"
@@ -190,15 +190,15 @@ void SStateMachineStateNode::GetNodeInfoPopups(FNodeInfoContext* InContext, TArr
 		return;
 	}
 
-	USceneStateGeneratedClass* Class = Cast<USceneStateGeneratedClass>(Blueprint->GeneratedClass);
-	if (!Class)
+	const FSceneStateTemplateData* TemplateData = RootState->GetSharedTemplateData().GetPtr<const FSceneStateTemplateData>();
+	if (!TemplateData)
 	{
 		return;
 	}
 
 	bIsActiveState = false;
 
-	Class->ForEachStateInstance(*RootState, GraphNode,
+	TemplateData->ForEachStateInstance(RootState->GetContextRegistry(), GraphNode,
 		[&OutPopups, &bIsActiveState = bIsActiveState, &ActiveColor = ActiveColor](const FSceneStateInstance& InInstance)
 		{
 			bIsActiveState = true;
diff --git a/Engine/Plugins/VirtualProduction/SceneState/Source/SceneStateMachineEditor/Private/Widgets/SSceneStateMachineTaskNode.cpp b/Engine/Plugins/VirtualProduction/SceneState/Source/SceneStateMachineEditor/Private/Widgets/SSceneStateMachineTaskNode.cpp
--- a/Engine/Plugins/VirtualProduction/SceneState/Source/SceneStateMachineEditor/Private/Widgets/SSceneStateMachineTaskNode.cpp
+++ b/Engine/Plugins/VirtualProduction/SceneState/Source/SceneStateMachineEditor/Private/Widgets/SSceneStateMachineTaskNode.cpp
@@ -7,10 +7,10 @@
 #include "SGraphPin.h"
 #include "SGraphPreviewer.h"
 #include "SSceneStateMachineOutputPin.h"
-#include "SceneStateGeneratedClass.h"
 #include "SceneStateMachineEditorStyle.h"
 #include "SceneStateMachineGraphSchema.h"
 #include "SceneStateObject.h"
+#include "SceneStateTemplateData.h"
 #include "Styling/SlateIconFinder.h"
 #include "Tasks/SceneStateBlueprintableTask.h"
 #include "Tasks/SceneStateTaskInstance.h"
@@ -190,15 +190,15 @@ void SStateMachineTaskNode::GetNodeInfoPopups(FNodeInfoContext* InContext, TArra
 		return;
 	}
 
-	USceneStateGeneratedClass* Class = Cast<USceneStateGeneratedClass>(Blueprint->GeneratedClass);
-	if (!Class)
+	const FSceneStateTemplateData* TemplateData = RootState->GetSharedTemplateData().GetPtr<const FSceneStateTemplateData>();
+	if (!TemplateData)
 	{
 		return;
 	}
 
 	TOptional<EExecutionStatus> ResultStatus;
 
-	Class->ForEachTaskInstance(*RootState, GraphNode,
+	TemplateData->ForEachTaskInstance(RootState->GetContextRegistry(), GraphNode,
 		[&OutPopups, &ResultStatus, &StatusColors = StatusColors](const FSceneStateTaskInstance& InTaskInstance)
 		{
 			// Get the lowest status to show
diff --git a/Engine/Plugins/VirtualProduction/SceneState/Source/SceneStateMachineEditor/Public/ISceneStateMachineCompilerContext.h b/Engine/Plugins/VirtualProduction/SceneState/Source/SceneStateMachineEditor/Public/ISceneStateMachineCompilerContext.h
--- a/Engine/Plugins/VirtualProduction/SceneState/Source/SceneStateMachineEditor/Public/ISceneStateMachineCompilerContext.h
+++ b/Engine/Plugins/VirtualProduction/SceneState/Source/SceneStateMachineEditor/Public/ISceneStateMachineCompilerContext.h
@@ -6,8 +6,8 @@
 #include "UObject/NameTypes.h"
 
 class UBlueprint;
-class USceneStateGeneratedClass;
 class USceneStateTransitionGraph;
+struct FSceneStateTemplateData;
 
 namespace UE::SceneState::Editor
 {
@@ -43,7 +43,11 @@ class IStateMachineCompilerContext
 public:
 	virtual UBlueprint* GetBlueprint() const = 0;
 
-	virtual USceneStateGeneratedClass* GetGeneratedClass() const = 0;
+	/** Gets the template data to compile the objects to */
+	virtual FSceneStateTemplateData* GetTemplateData() const = 0;
+
+	/** Returns the object to outer all compiled instanced UObjects to */
+	virtual UObject* GetTargetOuter() const = 0;
 
 	virtual FTransitionGraphCompileResult CompileTransitionGraph(USceneStateTransitionGraph* InTransitionGraph) = 0;
 };
diff --git a/Engine/Plugins/VirtualProduction/SceneState/Source/SceneStateMachineEditor/Public/SceneStateMachineCompiler.h b/Engine/Plugins/VirtualProduction/SceneState/Source/SceneStateMachineEditor/Public/SceneStateMachineCompiler.h
--- a/Engine/Plugins/VirtualProduction/SceneState/Source/SceneStateMachineEditor/Public/SceneStateMachineCompiler.h
+++ b/Engine/Plugins/VirtualProduction/SceneState/Source/SceneStateMachineEditor/Public/SceneStateMachineCompiler.h
@@ -13,12 +13,12 @@
 
 class UBlueprint;
 class UObject;
-class USceneStateGeneratedClass;
 class USceneStateMachineConduitNode;
 class USceneStateMachineGraph;
 class USceneStateMachineStateNode;
 class USceneStateMachineTransitionNode;
 struct FSceneStateMachine;
+struct FSceneStateTemplateData;
 
 namespace UE::SceneState::Editor
 {
@@ -67,7 +67,10 @@ private:
 	UObject* DuplicateObject(FObjectDuplicationParameters& InDuplicationParams);
 
 	/** Outer of the all the objects to create here */
-	TObjectPtr<USceneStateGeneratedClass> GeneratedClass;
+	TObjectPtr<UObject> Outer;
+
+	/** Template data where all the compiled objects will be allocated to */
+	FSceneStateTemplateData* TemplateData;
 
 	/** Reference Graph that will build out the runtime State Machine object */
 	TObjectPtr<USceneStateMachineGraph> StateMachineGraph;
