Fixed missing include

diff --git a/Engine/Plugins/VirtualProduction/SceneState/Source/SceneStateBlueprintEditor/Private/DebugExecutors/SceneStateDebugExecutor.h b/Engine/Plugins/VirtualProduction/SceneState/Source/SceneStateBlueprintEditor/Private/DebugExecutors/SceneStateDebugExecutor.h
--- a/Engine/Plugins/VirtualProduction/SceneState/Source/SceneStateBlueprintEditor/Private/DebugExecutors/SceneStateDebugExecutor.h
+++ b/Engine/Plugins/VirtualProduction/SceneState/Source/SceneStateBlueprintEditor/Private/DebugExecutors/SceneStateDebugExecutor.h
@@ -4,6 +4,7 @@
 
 #include "SceneStateExecutionContext.h"
 #include "Templates/SharedPointer.h"
+#include "UObject/GCObject.h"
 #include "UObject/ObjectKey.h"
 
 class USceneStateMachineNode;
