Motion Design Scene State: fixed issue where state machine parameters, etc would not appear in the property binding menu for funcitons. Additionally fixed an issue where binding extensions wre not allocating function instances in execution.

diff --git a/Engine/Plugins/VirtualProduction/SceneState/Source/SceneState/Private/Tasks/SceneStateTask.cpp b/Engine/Plugins/VirtualProduction/SceneState/Source/SceneState/Private/Tasks/SceneStateTask.cpp
--- a/Engine/Plugins/VirtualProduction/SceneState/Source/SceneState/Private/Tasks/SceneStateTask.cpp
+++ b/Engine/Plugins/VirtualProduction/SceneState/Source/SceneState/Private/Tasks/SceneStateTask.cpp
@@ -78,6 +78,16 @@ void FSceneStateTask::Setup(const FSceneStateExecutionContext& InContext, FStruc
 
 	Instance->SetStatus(EExecutionStatus::NotStarted);
 	InContext.SetupFunctionInstances(BindingsBatch);
+
+	if (const FSceneStateTaskBindingExtension* BindingExtension = GetBindingExtension())
+	{
+		BindingExtension->VisitBindingBatches(InTaskInstance,
+			[&InContext](uint16 InBatchIndex, FStructView InTargetDataView)
+			{
+				InContext.SetupFunctionInstances(FPropertyBindingIndex16(InBatchIndex));
+			});
+	}
+
 	OnSetup(InContext, InTaskInstance);
 }
 
@@ -139,6 +149,16 @@ void FSceneStateTask::Stop(const FSceneStateExecutionContext& InContext, FStruct
 	}
 
 	InContext.RemoveFunctionInstances(BindingsBatch);
+
+	if (const FSceneStateTaskBindingExtension* BindingExtension = GetBindingExtension())
+	{
+		BindingExtension->VisitBindingBatches(InTaskInstance,
+			[&InContext](uint16 InBatchIndex, FStructView InTargetDataView)
+			{
+				InContext.RemoveFunctionInstances(FPropertyBindingIndex16(InBatchIndex));
+			});
+	}
+
 	Instance->SetStatus(Private::GetTaskStatus(InStopReason));
 }
 
diff --git a/Engine/Plugins/VirtualProduction/SceneState/Source/SceneStateBinding/Public/SceneStateBindingCollectionOwner.h b/Engine/Plugins/VirtualProduction/SceneState/Source/SceneStateBinding/Public/SceneStateBindingCollectionOwner.h
--- a/Engine/Plugins/VirtualProduction/SceneState/Source/SceneStateBinding/Public/SceneStateBindingCollectionOwner.h
+++ b/Engine/Plugins/VirtualProduction/SceneState/Source/SceneStateBinding/Public/SceneStateBindingCollectionOwner.h
@@ -35,8 +35,5 @@ public:
 #if WITH_EDITOR
 	/** Iterates all functions that can be created and bound */
 	virtual bool ForEachBindableFunction(TFunctionRef<bool(const FSceneStateBindingDesc&, const UE::SceneState::FBindingFunctionInfo&)> InFunc) const = 0;
-
-	/** Iterates all functions found within the bindings */
-	virtual bool ForEachBindingFunction(TFunctionRef<bool(const FSceneStateBindingFunction&)> InFunc) const = 0;
 #endif
 };
diff --git a/Engine/Plugins/VirtualProduction/SceneState/Source/SceneStateBlueprint/Private/SceneStateBlueprint.cpp b/Engine/Plugins/VirtualProduction/SceneState/Source/SceneStateBlueprint/Private/SceneStateBlueprint.cpp
--- a/Engine/Plugins/VirtualProduction/SceneState/Source/SceneStateBlueprint/Private/SceneStateBlueprint.cpp
+++ b/Engine/Plugins/VirtualProduction/SceneState/Source/SceneStateBlueprint/Private/SceneStateBlueprint.cpp
@@ -61,24 +61,20 @@ FSceneStateBindingDesc USceneStateBlueprint::CreateRootBinding() const
 	return CreateBindingDesc(*this);
 }
 
-void USceneStateBlueprint::GetBindingStructs(const FGetBindableStructsParams& InParams, TArray<TInstancedStruct<FPropertyBindingBindableStructDescriptor>>& OutStructDescs) const
+void USceneStateBlueprint::GetBindingStructs(const FGetBindableStructsParams& InParams, TArray<TInstancedStruct<FSceneStateBindingDesc>>& OutBindingDescs) const
 {
 	using namespace UE::SceneState::Graph;
 
-	TArray<TInstancedStruct<FSceneStateBindingDesc>> BindingDescs;
-
-	ON_SCOPE_EXIT
-	{
-		OutStructDescs.Append(MoveTemp(BindingDescs));
-	};
-
 	auto GatherDefaultBindingDescs =
-		[This=this, &InParams, &BindingDescs]()
+		[This=this, &InParams, &OutBindingDescs]()
 		{
-			GetGlobalBindingDescs(*This, InParams.TargetStructId, BindingDescs);
+			if (InParams.bIncludeGlobalDescs)
+			{
+				GetGlobalBindingDescs(*This, InParams.TargetStructId, OutBindingDescs);	
+			}
 			if (InParams.bIncludeFunctions)
 			{
-				GetFunctionBindingDescs(*This, InParams.TargetStructId, BindingDescs);
+				GetFunctionBindingDescs(*This, InParams.TargetStructId, OutBindingDescs);
 			}
 		};
 
@@ -86,7 +82,7 @@ void USceneStateBlueprint::GetBindingStructs(const FGetBindableStructsParams& In
 	if (USceneStateMachineGraph* SceneStateMachine = FindStateMachineMatchingId(*this, InParams.TargetStructId))
 	{
 		GatherDefaultBindingDescs();
-		GetStateMachineBindingDescs(*this, *SceneStateMachine, BindingDescs);
+		GetStateMachineBindingDescs(*this, *SceneStateMachine, OutBindingDescs);
 		return;
 	}
 
@@ -94,7 +90,7 @@ void USceneStateBlueprint::GetBindingStructs(const FGetBindableStructsParams& In
 	if (USceneStateMachineTransitionNode* TransitionNode = FindTransitionMatchingId(*this, InParams.TargetStructId))
 	{
 		GatherDefaultBindingDescs();
-		GetTransitionBindingDescs(*this, *TransitionNode, BindingDescs);
+		GetTransitionBindingDescs(*this, *TransitionNode, OutBindingDescs);
 		return;
 	}
 
@@ -105,17 +101,43 @@ void USceneStateBlueprint::GetBindingStructs(const FGetBindableStructsParams& In
 	if (const USceneStateMachineTaskNode* TargetTaskNode = FindTaskNodeContainingId(*this, InParams.TargetStructId))
 	{
 		GatherDefaultBindingDescs();
-		GetTaskBindingDescs(*this, *TargetTaskNode, BindingDescs);
+		GetTaskBindingDescs(*this, *TargetTaskNode, OutBindingDescs);
 		return;
 	}
 
 	// Add only default binding descs if the target struct is a binding function
-	if (const FSceneStateBindingFunction* BindingFunction = FindBindingFunctionMatchingId(*this, InParams.TargetStructId))
+	FGuid OwnerStructId;
+	if (const FSceneStateBindingFunction* BindingFunction = FindBindingFunctionMatchingId(*this, InParams.TargetStructId, &OwnerStructId))
 	{
 		GatherDefaultBindingDescs();
+
+		// Gather the Binding Structs for the Owner
+		const FGetBindableStructsParams Params
+		{
+			.TargetStructId = OwnerStructId,
+			.bIncludeGlobalDescs = false,
+			.bIncludeFunctions = InParams.bIncludeFunctions,
+		};
+		GetBindingStructs(Params, OutBindingDescs);
 	}
 }
 
+bool USceneStateBlueprint::ForEachBindingFunction(TFunctionRef<bool(const FSceneStateBindingFunction&, const FPropertyBindingBinding&)> InFunc) const
+{
+	for (const FPropertyBindingBinding& Binding : BindingCollection.GetBindings())
+	{
+		const FSceneStateBindingFunction* BindingFunction = Binding.GetPropertyFunctionNode().GetPtr<const FSceneStateBindingFunction>();
+		if (BindingFunction && BindingFunction->IsValid())
+		{
+			if (!InFunc(*BindingFunction, Binding))
+			{
+				return false;
+			}
+		}
+	}
+	return true;
+}
+
 FSceneStateBindingCollection& USceneStateBlueprint::GetBindingCollection()
 {
 	return BindingCollection;
@@ -164,22 +186,6 @@ bool USceneStateBlueprint::ForEachBindableFunction(TFunctionRef<bool(const FScen
 	return true;
 }
 
-bool USceneStateBlueprint::ForEachBindingFunction(TFunctionRef<bool(const FSceneStateBindingFunction&)> InFunc) const
-{
-	for (const FPropertyBindingBinding& Binding : BindingCollection.GetBindings())
-	{
-		const FSceneStateBindingFunction* BindingFunction = Binding.GetPropertyFunctionNode().GetPtr<const FSceneStateBindingFunction>();
-		if (BindingFunction && BindingFunction->IsValid())
-		{
-			if (!InFunc(*BindingFunction))
-			{
-				return false;
-			}
-		}
-	}
-	return true;
-}
-
 void USceneStateBlueprint::GetBindableStructs(const FGuid InTargetStructId, TArray<TInstancedStruct<FPropertyBindingBindableStructDescriptor>>& OutStructDescs) const
 {
 	const FGetBindableStructsParams Params
@@ -187,7 +193,10 @@ void USceneStateBlueprint::GetBindableStructs(const FGuid InTargetStructId, TArr
 			.TargetStructId = InTargetStructId,
 			.bIncludeFunctions = false,
 		};
-	GetBindingStructs(Params, OutStructDescs);
+
+	TArray<TInstancedStruct<FSceneStateBindingDesc>> BindingDescs;
+	GetBindingStructs(Params, BindingDescs);
+	OutStructDescs.Append(MoveTemp(BindingDescs));
 }
 
 bool USceneStateBlueprint::GetBindableStructByID(const FGuid InStructId, TInstancedStruct<FPropertyBindingBindableStructDescriptor>& OutStructDesc) const
diff --git a/Engine/Plugins/VirtualProduction/SceneState/Source/SceneStateBlueprint/Private/SceneStateBlueprintBindingUtils.cpp b/Engine/Plugins/VirtualProduction/SceneState/Source/SceneStateBlueprint/Private/SceneStateBlueprintBindingUtils.cpp
--- a/Engine/Plugins/VirtualProduction/SceneState/Source/SceneStateBlueprint/Private/SceneStateBlueprintBindingUtils.cpp
+++ b/Engine/Plugins/VirtualProduction/SceneState/Source/SceneStateBlueprint/Private/SceneStateBlueprintBindingUtils.cpp
@@ -197,16 +197,20 @@ void GetTaskBindingDescs(const USceneStateBlueprint& InBlueprint, const USceneSt
 	}
 }
 
-const FSceneStateBindingFunction* FindBindingFunctionMatchingId(const USceneStateBlueprint& InBlueprint, const FGuid& InStructId)
+const FSceneStateBindingFunction* FindBindingFunctionMatchingId(const USceneStateBlueprint& InBlueprint, const FGuid& InStructId, FGuid* OutOwnerStructId)
 {
 	const FSceneStateBindingFunction* Result = nullptr;
 
 	InBlueprint.ForEachBindingFunction(
-		[&InStructId, &Result](const FSceneStateBindingFunction& InBindingFunction)
+		[&InStructId, &Result, OutOwnerStructId](const FSceneStateBindingFunction& InBindingFunction, const FPropertyBindingBinding& InBinding)
 		{
 			if (InBindingFunction.FunctionId == InStructId)
 			{
 				Result = &InBindingFunction;
+				if (OutOwnerStructId)
+				{
+					*OutOwnerStructId = InBinding.GetTargetPath().GetStructID();
+				}
 				return false; // break
 			}
 			return true; // continue
diff --git a/Engine/Plugins/VirtualProduction/SceneState/Source/SceneStateBlueprint/Private/SceneStateBlueprintBindingUtils.h b/Engine/Plugins/VirtualProduction/SceneState/Source/SceneStateBlueprint/Private/SceneStateBlueprintBindingUtils.h
--- a/Engine/Plugins/VirtualProduction/SceneState/Source/SceneStateBlueprint/Private/SceneStateBlueprintBindingUtils.h
+++ b/Engine/Plugins/VirtualProduction/SceneState/Source/SceneStateBlueprint/Private/SceneStateBlueprintBindingUtils.h
@@ -91,9 +91,10 @@ namespace UE::SceneState::Graph
 	 * Finds the binding function that matches the given struct id.
 	 * @param InBlueprint the blueprint to look into
 	 * @param InStructId the struct id to look for
-	 * @return the found binding function containing the given struct id
+	 * @param OutOwnerStructId (optional) the found binding function's owner struct id
+	 * @return the found binding function matching the given struct id
 	 */
-	const FSceneStateBindingFunction* FindBindingFunctionMatchingId(const USceneStateBlueprint& InBlueprint, const FGuid& InStructId);
+	const FSceneStateBindingFunction* FindBindingFunctionMatchingId(const USceneStateBlueprint& InBlueprint, const FGuid& InStructId, FGuid* OutOwnerStructId = nullptr);
 
 	/**
 	 * Finds the transition node with id that matches the given struct id.
diff --git a/Engine/Plugins/VirtualProduction/SceneState/Source/SceneStateBlueprint/Public/SceneStateBlueprint.h b/Engine/Plugins/VirtualProduction/SceneState/Source/SceneStateBlueprint/Public/SceneStateBlueprint.h
--- a/Engine/Plugins/VirtualProduction/SceneState/Source/SceneStateBlueprint/Public/SceneStateBlueprint.h
+++ b/Engine/Plugins/VirtualProduction/SceneState/Source/SceneStateBlueprint/Public/SceneStateBlueprint.h
@@ -40,17 +40,21 @@ public:
 	{
 		/** Struct id to look bindable structs for */
 		FGuid TargetStructId;
-		/** Whether to include the binding functions bound to the struct. Useful for compiling, but not */
+		/** Whether to include the global bindings to the struct */
+		bool bIncludeGlobalDescs = true;
+		/** Whether to include the binding functions bound to the struct */
 		bool bIncludeFunctions = false;
 	};
 	/** Gathers the binding structs for a given struct id. */
-	UE_API void GetBindingStructs(const FGetBindableStructsParams& InParams, TArray<TInstancedStruct<FPropertyBindingBindableStructDescriptor>>& OutStructDescs) const;
+	UE_API void GetBindingStructs(const FGetBindableStructsParams& InParams, TArray<TInstancedStruct<FSceneStateBindingDesc>>& OutBindingDescs) const;
+
+	/** Iterates all functions found within the bindings */
+	bool ForEachBindingFunction(TFunctionRef<bool(const FSceneStateBindingFunction&, const FPropertyBindingBinding&)> InFunc) const;
 
 	//~ Begin ISceneStateBindingCollectionOwner
 	UE_API virtual FSceneStateBindingCollection& GetBindingCollection() override;
 	UE_API virtual const FSceneStateBindingCollection& GetBindingCollection() const override;
 	UE_API virtual bool ForEachBindableFunction(TFunctionRef<bool(const FSceneStateBindingDesc&, const UE::SceneState::FBindingFunctionInfo&)> InFunc) const override;
-	UE_API virtual bool ForEachBindingFunction(TFunctionRef<bool(const FSceneStateBindingFunction&)> InFunc) const;
 	//~ End ISceneStateBindingCollectionOwner
 
 	//~ Begin IPropertyBindingBindingCollectionOwner
diff --git a/Engine/Plugins/VirtualProduction/SceneState/Source/SceneStateBlueprintEditor/Private/Compiler/SceneStateBindingCompiler.cpp b/Engine/Plugins/VirtualProduction/SceneState/Source/SceneStateBlueprintEditor/Private/Compiler/SceneStateBindingCompiler.cpp
--- a/Engine/Plugins/VirtualProduction/SceneState/Source/SceneStateBlueprintEditor/Private/Compiler/SceneStateBindingCompiler.cpp
+++ b/Engine/Plugins/VirtualProduction/SceneState/Source/SceneStateBlueprintEditor/Private/Compiler/SceneStateBindingCompiler.cpp
@@ -264,7 +264,7 @@ bool FBindingCompiler::ValidateBinding(const FSceneStateBinding& InBinding) cons
 
 	// Source must be accessible by the target struct.
 	// This mismatch could happen if copying an object bound to a scoped parameter and pasting it outside such scope
-	TArray<TInstancedStruct<FPropertyBindingBindableStructDescriptor>> AccessibleStructs;
+	TArray<TInstancedStruct<FSceneStateBindingDesc>> AccessibleStructs;
 	Blueprint->GetBindingStructs(Params, AccessibleStructs);
 
 	const FGuid& SourceStructId = InBinding.GetSourcePath().GetStructID();
