Motion Design Scene State: removed unused interface function. Blueprint (if any) is automatically gotten from the graph.

diff --git a/Engine/Plugins/VirtualProduction/SceneState/Source/SceneStateBlueprintEditor/Private/Compiler/SceneStateBlueprintCompilerContext.cpp b/Engine/Plugins/VirtualProduction/SceneState/Source/SceneStateBlueprintEditor/Private/Compiler/SceneStateBlueprintCompilerContext.cpp
--- a/Engine/Plugins/VirtualProduction/SceneState/Source/SceneStateBlueprintEditor/Private/Compiler/SceneStateBlueprintCompilerContext.cpp
+++ b/Engine/Plugins/VirtualProduction/SceneState/Source/SceneStateBlueprintEditor/Private/Compiler/SceneStateBlueprintCompilerContext.cpp
@@ -188,11 +188,6 @@ void FBlueprintCompilerContext::OnPostCDOCompiled(const UObject::FPostCDOCompile
 	NewGeneratedClass->ResolveBindings();
 }
 
-UBlueprint* FBlueprintCompilerContext::GetBlueprint() const
-{
-	return Blueprint;
-}
-
 FSceneStateTemplateData* FBlueprintCompilerContext::GetTemplateData() const
 {
 	return NewGeneratedClass
diff --git a/Engine/Plugins/VirtualProduction/SceneState/Source/SceneStateBlueprintEditor/Private/Compiler/SceneStateBlueprintCompilerContext.h b/Engine/Plugins/VirtualProduction/SceneState/Source/SceneStateBlueprintEditor/Private/Compiler/SceneStateBlueprintCompilerContext.h
--- a/Engine/Plugins/VirtualProduction/SceneState/Source/SceneStateBlueprintEditor/Private/Compiler/SceneStateBlueprintCompilerContext.h
+++ b/Engine/Plugins/VirtualProduction/SceneState/Source/SceneStateBlueprintEditor/Private/Compiler/SceneStateBlueprintCompilerContext.h
@@ -28,7 +28,6 @@ private:
 	//~ End FKismetCompilerContext
 
 	//~ Begin ISceneStateMachineCompilerContext
-	virtual UBlueprint* GetBlueprint() const override;
 	virtual FSceneStateTemplateData* GetTemplateData() const override;
 	virtual UObject* GetTargetOuter() const override;
 	virtual FTransitionGraphCompileResult CompileTransitionGraph(USceneStateTransitionGraph* InTransitionGraph) override;
diff --git a/Engine/Plugins/VirtualProduction/SceneState/Source/SceneStateMachineEditor/Public/ISceneStateMachineCompilerContext.h b/Engine/Plugins/VirtualProduction/SceneState/Source/SceneStateMachineEditor/Public/ISceneStateMachineCompilerContext.h
--- a/Engine/Plugins/VirtualProduction/SceneState/Source/SceneStateMachineEditor/Public/ISceneStateMachineCompilerContext.h
+++ b/Engine/Plugins/VirtualProduction/SceneState/Source/SceneStateMachineEditor/Public/ISceneStateMachineCompilerContext.h
@@ -5,7 +5,6 @@
 #include "HAL/Platform.h"
 #include "UObject/NameTypes.h"
 
-class UBlueprint;
 class USceneStateTransitionGraph;
 struct FSceneStateTemplateData;
 
@@ -41,8 +40,6 @@ struct FTransitionGraphCompileResult
 class IStateMachineCompilerContext
 {
 public:
-	virtual UBlueprint* GetBlueprint() const = 0;
-
 	/** Gets the template data to compile the objects to */
 	virtual FSceneStateTemplateData* GetTemplateData() const = 0;
 
diff --git a/Engine/Plugins/VirtualProduction/SceneState/Source/SceneStateMachineEditor/Public/SceneStateMachineCompiler.h b/Engine/Plugins/VirtualProduction/SceneState/Source/SceneStateMachineEditor/Public/SceneStateMachineCompiler.h
--- a/Engine/Plugins/VirtualProduction/SceneState/Source/SceneStateMachineEditor/Public/SceneStateMachineCompiler.h
+++ b/Engine/Plugins/VirtualProduction/SceneState/Source/SceneStateMachineEditor/Public/SceneStateMachineCompiler.h
@@ -75,6 +75,7 @@ private:
 	/** Reference Graph that will build out the runtime State Machine object */
 	TObjectPtr<USceneStateMachineGraph> StateMachineGraph;
 
+	/** Blueprint owning the state machine graph, if any */
 	TObjectPtr<UBlueprint> Blueprint;
 
 	/** Context used for compiling the State Machine Graph */
