Motion Design Scene State: fixed issue where copying tasks/etc would not copy over the function values (only type).

diff --git a/Engine/Plugins/VirtualProduction/SceneState/Source/SceneStateBinding/Private/SceneStateBindingCollection.cpp b/Engine/Plugins/VirtualProduction/SceneState/Source/SceneStateBinding/Private/SceneStateBindingCollection.cpp
--- a/Engine/Plugins/VirtualProduction/SceneState/Source/SceneStateBinding/Private/SceneStateBindingCollection.cpp
+++ b/Engine/Plugins/VirtualProduction/SceneState/Source/SceneStateBinding/Private/SceneStateBindingCollection.cpp
@@ -86,8 +86,8 @@ void FSceneStateBindingCollection::CopyBindingsInternal(const FGuid InFromStruct
 				if (const FSceneStateBindingFunction* BindingFunction = InBinding.GetPropertyFunctionNode().GetPtr<const FSceneStateBindingFunction>())
 				{
 					FFunctionBindingCopyInfo& CopyInfo = FunctionBindingsToCopy.AddDefaulted_GetRef();
-					CopyInfo.FunctionInfo.FunctionType = BindingFunction->Function.GetScriptStruct();
-					CopyInfo.FunctionInfo.InstanceType = BindingFunction->FunctionInstance.GetScriptStruct();
+					CopyInfo.FunctionInfo.FunctionTemplate = BindingFunction->Function;
+					CopyInfo.FunctionInfo.InstanceTemplate = BindingFunction->FunctionInstance;
 					CopyInfo.SourcePath = InBinding.GetSourcePath();
 					CopyInfo.TargetPath = FPropertyBindingPath(ToStructId, InBinding.GetTargetPath().GetSegments());
 
diff --git a/Engine/Plugins/VirtualProduction/SceneState/Source/SceneStateBinding/Private/SceneStateBindingFunction.cpp b/Engine/Plugins/VirtualProduction/SceneState/Source/SceneStateBinding/Private/SceneStateBindingFunction.cpp
--- a/Engine/Plugins/VirtualProduction/SceneState/Source/SceneStateBinding/Private/SceneStateBindingFunction.cpp
+++ b/Engine/Plugins/VirtualProduction/SceneState/Source/SceneStateBinding/Private/SceneStateBindingFunction.cpp
@@ -6,8 +6,8 @@
 FSceneStateBindingFunction::FSceneStateBindingFunction(const UE::SceneState::FBindingFunctionInfo& InFunctionInfo)
 {
 	FunctionId = FGuid::NewGuid();
-	Function.InitializeAs(InFunctionInfo.FunctionType);
-	FunctionInstance.InitializeAs(InFunctionInfo.InstanceType);
+	Function = InFunctionInfo.FunctionTemplate;
+	FunctionInstance = InFunctionInfo.InstanceTemplate;
 }
 
 bool FSceneStateBindingFunction::IsValid() const
diff --git a/Engine/Plugins/VirtualProduction/SceneState/Source/SceneStateBinding/Public/SceneStateBindingFunction.h b/Engine/Plugins/VirtualProduction/SceneState/Source/SceneStateBinding/Public/SceneStateBindingFunction.h
--- a/Engine/Plugins/VirtualProduction/SceneState/Source/SceneStateBinding/Public/SceneStateBindingFunction.h
+++ b/Engine/Plugins/VirtualProduction/SceneState/Source/SceneStateBinding/Public/SceneStateBindingFunction.h
@@ -3,6 +3,7 @@
 #pragma once
 
 #include "StructUtils/InstancedStruct.h"
+#include "StructUtils/StructView.h"
 #include "SceneStateBindingFunction.generated.h"
 
 #define UE_API SCENESTATEBINDING_API
@@ -14,8 +15,10 @@ namespace UE::SceneState
 /** Arguments to create a valid FSceneStateBindingFunction */
 struct FBindingFunctionInfo
 {
-	const UScriptStruct* FunctionType = nullptr;
-	const UScriptStruct* InstanceType = nullptr;
+	/** Info to create a new function. Struct type is required. The struct memory is optional and only used as a template to instantiate the function */
+	FConstStructView FunctionTemplate;
+	/** Info to create a new function instance. Struct type is required. The struct memory is optional and only used as a template to instantiate the function instance */
+	FConstStructView InstanceTemplate;
 };
 #endif
 
diff --git a/Engine/Plugins/VirtualProduction/SceneState/Source/SceneStateBlueprint/Private/SceneStateBlueprint.cpp b/Engine/Plugins/VirtualProduction/SceneState/Source/SceneStateBlueprint/Private/SceneStateBlueprint.cpp
--- a/Engine/Plugins/VirtualProduction/SceneState/Source/SceneStateBlueprint/Private/SceneStateBlueprint.cpp
+++ b/Engine/Plugins/VirtualProduction/SceneState/Source/SceneStateBlueprint/Private/SceneStateBlueprint.cpp
@@ -174,8 +174,8 @@ bool USceneStateBlueprint::ForEachBindableFunction(TFunctionRef<bool(const FScen
 		BindingDesc.DataHandle = FSceneStateBindingDataHandle(ESceneStateDataType::Function);
 
 		UE::SceneState::FBindingFunctionInfo FunctionInfo;
-		FunctionInfo.FunctionType = Struct;
-		FunctionInfo.InstanceType = InstanceType;
+		FunctionInfo.FunctionTemplate = FConstStructView(Struct);
+		FunctionInfo.InstanceTemplate = FConstStructView(InstanceType);
 
 		if (!InFunc(BindingDesc, FunctionInfo))
 		{
diff --git a/Engine/Plugins/VirtualProduction/SceneState/Source/SceneStateBlueprintEditor/Private/DetailsView/SceneStateCachedBindingData.cpp b/Engine/Plugins/VirtualProduction/SceneState/Source/SceneStateBlueprintEditor/Private/DetailsView/SceneStateCachedBindingData.cpp
--- a/Engine/Plugins/VirtualProduction/SceneState/Source/SceneStateBlueprintEditor/Private/DetailsView/SceneStateCachedBindingData.cpp
+++ b/Engine/Plugins/VirtualProduction/SceneState/Source/SceneStateBlueprintEditor/Private/DetailsView/SceneStateCachedBindingData.cpp
@@ -286,7 +286,7 @@ bool FSceneStateCachedBindingData::AddBindingInternal(TConstStructView<FProperty
 			return true; // continue
 		});
 
-	if (!ensure(FunctionInfo.FunctionType))
+	if (!ensure(FunctionInfo.FunctionTemplate.GetScriptStruct()))
 	{
 		return false;
 	}
