Motion Design Scene State: fixed issue where Transition Parameter enums were unusable in the graph

diff --git a/Engine/Plugins/VirtualProduction/SceneState/Source/SceneStateBlueprintEditor/Private/Compiler/SceneStateTransitionGraphCompiler.cpp b/Engine/Plugins/VirtualProduction/SceneState/Source/SceneStateBlueprintEditor/Private/Compiler/SceneStateTransitionGraphCompiler.cpp
--- a/Engine/Plugins/VirtualProduction/SceneState/Source/SceneStateBlueprintEditor/Private/Compiler/SceneStateTransitionGraphCompiler.cpp
+++ b/Engine/Plugins/VirtualProduction/SceneState/Source/SceneStateBlueprintEditor/Private/Compiler/SceneStateTransitionGraphCompiler.cpp
@@ -9,6 +9,7 @@
 #include "PropertyBagDetails.h"
 #include "SceneStateBlueprintCompilerContext.h"
 #include "SceneStateMachineGraph.h"
+#include "SceneStateMachineGraphUtils.h"
 #include "SceneStateTransitionGraph.h"
 #include "Transition/SceneStateTransitionParametersNode.h"
 
@@ -255,7 +256,7 @@ UK2Node_CustomEvent* FTransitionGraphCompiler::CreateCustomEventNode()
 
 	for (const FPropertyBagPropertyDesc& PropertyDesc : PropertyDescs)
 	{
-		const FEdGraphPinType PinType = UE::StructUtils::GetPropertyDescAsPin(PropertyDesc);
+		const FEdGraphPinType PinType = UE::SceneState::Graph::GetPropertyDescAsPin(PropertyDesc);
 		CustomEventNode->CreateUserDefinedPin(PropertyDesc.Name, PinType, EGPD_Output, /*bUseUniqueName*/false);
 	}
 
diff --git a/Engine/Plugins/VirtualProduction/SceneState/Source/SceneStateMachineGraph/Private/SceneStateMachineGraphUtils.cpp b/Engine/Plugins/VirtualProduction/SceneState/Source/SceneStateMachineGraph/Private/SceneStateMachineGraphUtils.cpp
--- a/Engine/Plugins/VirtualProduction/SceneState/Source/SceneStateMachineGraph/Private/SceneStateMachineGraphUtils.cpp
+++ b/Engine/Plugins/VirtualProduction/SceneState/Source/SceneStateMachineGraph/Private/SceneStateMachineGraphUtils.cpp
@@ -3,6 +3,7 @@
 #include "SceneStateMachineGraphUtils.h"
 #include "Kismet2/BlueprintEditorUtils.h"
 #include "Nodes/SceneStateMachineNode.h"
+#include "PropertyBagDetails.h"
 
 namespace UE::SceneState::Graph
 {
@@ -45,4 +46,18 @@ void RemoveGraph(UEdGraph* InGraph)
 	}
 }
 
+FEdGraphPinType GetPropertyDescAsPin(const FPropertyBagPropertyDesc& InPropertyDesc)
+{
+	FEdGraphPinType PinType = UE::StructUtils::GetPropertyDescAsPin(InPropertyDesc);
+
+	// Change enum category to byte. Even though enum is more accurate, existing code in blueprint API only checks for PC_Byte.
+	// Cases like UEdGraphSchema_K2::GetPropertyCategoryInfo also set the pin category to byte for enum properties.
+	if (PinType.PinCategory == UEdGraphSchema_K2::PC_Enum)
+	{
+		PinType.PinCategory = UEdGraphSchema_K2::PC_Byte;
+	}
+
+	return PinType;
+}
+
 } //UE::SceneState::Graph
diff --git a/Engine/Plugins/VirtualProduction/SceneState/Source/SceneStateMachineGraph/Private/Transition/SceneStateTransitionParametersNode.cpp b/Engine/Plugins/VirtualProduction/SceneState/Source/SceneStateMachineGraph/Private/Transition/SceneStateTransitionParametersNode.cpp
--- a/Engine/Plugins/VirtualProduction/SceneState/Source/SceneStateMachineGraph/Private/Transition/SceneStateTransitionParametersNode.cpp
+++ b/Engine/Plugins/VirtualProduction/SceneState/Source/SceneStateMachineGraph/Private/Transition/SceneStateTransitionParametersNode.cpp
@@ -6,6 +6,7 @@
 #include "Kismet2/KismetDebugUtilities.h"
 #include "Nodes/SceneStateMachineTransitionNode.h"
 #include "PropertyBagDetails.h"
+#include "SceneStateMachineGraphUtils.h"
 #include "SceneStateTransitionGraphSchema.h"
 
 #define LOCTEXT_NAMESPACE "SceneStateTransitionParametersNode"
@@ -50,9 +51,8 @@ void USceneStateTransitionParametersNode::BuildParameterPins()
 
 	for (const FPropertyBagPropertyDesc& PropertyDesc : PropertyDescs)
 	{
-		const FEdGraphPinType PinType = UE::StructUtils::GetPropertyDescAsPin(PropertyDesc);
+		const FEdGraphPinType PinType = UE::SceneState::Graph::GetPropertyDescAsPin(PropertyDesc);
 		CreateUserDefinedPin(PropertyDesc.Name, PinType, EGPD_Output, /*bUseUniqueName*/false);
-		// todo: default values for pin
 	}
 }
 
diff --git a/Engine/Plugins/VirtualProduction/SceneState/Source/SceneStateMachineGraph/Public/SceneStateMachineGraphUtils.h b/Engine/Plugins/VirtualProduction/SceneState/Source/SceneStateMachineGraph/Public/SceneStateMachineGraphUtils.h
--- a/Engine/Plugins/VirtualProduction/SceneState/Source/SceneStateMachineGraph/Public/SceneStateMachineGraphUtils.h
+++ b/Engine/Plugins/VirtualProduction/SceneState/Source/SceneStateMachineGraph/Public/SceneStateMachineGraphUtils.h
@@ -4,12 +4,21 @@
 
 #include "HAL/Platform.h"
 
+#define UE_API SCENESTATEMACHINEGRAPH_API
+
 class UEdGraph;
+struct FEdGraphPinType;
+struct FPropertyBagPropertyDesc;
 
 namespace UE::SceneState::Graph
 {
-	SCENESTATEMACHINEGRAPH_API bool CanDirectlyRemoveGraph(UEdGraph* InGraph);
+	UE_API bool CanDirectlyRemoveGraph(UEdGraph* InGraph);
+
+	UE_API void RemoveGraph(UEdGraph* InGraph);
 
-	SCENESTATEMACHINEGRAPH_API void RemoveGraph(UEdGraph* InGraph);
+	/** Gets the given property desc as a pin compatible for blueprints */
+	UE_API FEdGraphPinType GetPropertyDescAsPin(const FPropertyBagPropertyDesc& InPropertyDesc);
 
 } // UE::SceneState::Graph
+
+#undef UE_API
