Motion Design Scene State: fixed issue where uobjects instanced to the generated class and saved in the shared struct were getting marked as unreachable and garbage collected.
This was done by changing the template data from being a struct (was being used as a shared struct) to a uobject. This uobject approach has the benefit that task uobjects can now be outered to the template data uobject directly rather than the owning class.

diff --git a/Engine/Plugins/VirtualProduction/SceneState/Source/SceneState/Private/SceneStateExecutionContext.cpp b/Engine/Plugins/VirtualProduction/SceneState/Source/SceneState/Private/SceneStateExecutionContext.cpp
--- a/Engine/Plugins/VirtualProduction/SceneState/Source/SceneState/Private/SceneStateExecutionContext.cpp
+++ b/Engine/Plugins/VirtualProduction/SceneState/Source/SceneState/Private/SceneStateExecutionContext.cpp
@@ -22,11 +22,6 @@ FSceneStateExecutionContext::~FSceneStateExecutionContext()
 	Reset();
 }
 
-const FSceneStateTemplateData& FSceneStateExecutionContext::GetSharedTemplateData() const
-{
-	return SharedTemplateData.Get<const FSceneStateTemplateData>();
-}
-
 TSharedPtr<const UE::SceneState::FExecutionContextRegistry> FSceneStateExecutionContext::GetContextRegistry() const
 {
 	return ContextRegistryWeak.Pin();
@@ -34,7 +29,7 @@ TSharedPtr<const UE::SceneState::FExecutionContextRegistry> FSceneStateExecution
 
 const FSceneStateBindingCollection& FSceneStateExecutionContext::GetBindingCollection() const
 {
-	return GetSharedTemplateData().BindingCollection;
+	return TemplateData->BindingCollection;
 }
 
 const FString& FSceneStateExecutionContext::GetExecutionContextName() const
@@ -61,7 +56,7 @@ void FSceneStateExecutionContext::Setup(const FSceneStateExecutionContext& InPar
 {
 	const FSetupParams SetupParams
 	{
-		.SharedTemplateData = InParentContext.SharedTemplateData,
+		.TemplateData = InParentContext.TemplateData,
 		.RootObject = InParentContext.GetRootObject(),
 		.ContextObject = InParentContext.GetContextObject(),
 		.EventStream = InParentContext.GetEventStream(),
@@ -80,7 +75,7 @@ void FSceneStateExecutionContext::Setup(TNotNull<USceneStateObject*> InRootState
 {
 	const FSetupParams SetupParams
 	{
-		.SharedTemplateData = InRootState->GetSharedTemplateData(),
+		.TemplateData = InRootState->GetTemplateData(),
 		.RootObject = InRootState,
 		.ContextObject = InRootState->GetContextObject(),
 		.EventStream = InRootState->GetEventStream(),
@@ -94,7 +89,7 @@ void FSceneStateExecutionContext::Setup(const FSetupParams& InParams)
 {
 	Reset();
 
-	SharedTemplateData = InParams.SharedTemplateData;
+	TemplateData = InParams.TemplateData;
 	RootObject = InParams.RootObject;
 	ContextObjectWeak = InParams.ContextObject;
 	EventStreamWeak = InParams.EventStream;
@@ -109,7 +104,7 @@ void FSceneStateExecutionContext::Setup(const FSetupParams& InParams)
 
 void FSceneStateExecutionContext::Reset()
 {
-	SharedTemplateData.Reset();
+	TemplateData = nullptr;
 
 	StateInstances.Reset();
 	StateMachineInstances.Reset();
@@ -217,9 +212,7 @@ void FSceneStateExecutionContext::ForEachTask(const FSceneState& InState, TFunct
 		return;
 	}
 
-	const FSceneStateTemplateData& TemplateData = GetSharedTemplateData();
-
-	const FInstancedStructContainer& Tasks = TemplateData.Tasks;
+	const FInstancedStructContainer& Tasks = TemplateData->Tasks;
 
 	checkfSlow(Tasks.IsValidIndex(TaskRange.Index) && Tasks.IsValidIndex(TaskRange.GetLastIndex())
 		, TEXT("[Context: %s] State Task Range [%d, %d] is out of bounds! (Total Task Count: %d)")
@@ -248,14 +241,7 @@ void FSceneStateExecutionContext::ForEachTask(const FSceneState& InState, TFunct
 
 TArray<FConstStructView> FSceneStateExecutionContext::GetTemplateTaskInstances(const FSceneState& InState) const
 {
-	const FSceneStateRange TaskRange = InState.GetTaskRange();
-	if (TaskRange.Count == 0)
-	{
-		return {};
-	}
-
-	const FSceneStateTemplateData& TemplateData = GetSharedTemplateData();
-	return UE::SceneState::GetConstStructViews(TemplateData.TaskInstances, TaskRange);
+	return UE::SceneState::GetConstStructViews(TemplateData->TaskInstances, InState.GetTaskRange());
 }
 
 TConstArrayView<uint16> FSceneStateExecutionContext::GetTaskPrerequisites(const FSceneStateTask& InTask) const
@@ -266,8 +252,7 @@ TConstArrayView<uint16> FSceneStateExecutionContext::GetTaskPrerequisites(const
 		return {};
 	}
 
-	const FSceneStateTemplateData& TemplateData = GetSharedTemplateData();
-	const TConstArrayView<uint16> TaskPrerequisitesView = TemplateData.TaskPrerequisites;
+	const TConstArrayView<uint16> TaskPrerequisitesView = TemplateData->TaskPrerequisites;
 
 	checkfSlow(TaskPrerequisitesView.IsValidIndex(TaskPrerequisiteRange.Index) && TaskPrerequisitesView.IsValidIndex(TaskPrerequisiteRange.GetLastIndex())
 		, TEXT("[Context: %s] Task Prerequisite Range [%d, %d] is out of bounds! (Total Prerequisite Count: %d)")
@@ -287,8 +272,7 @@ TConstArrayView<FSceneStateTransition> FSceneStateExecutionContext::GetTransitio
 		return {};
 	}
 
-	const FSceneStateTemplateData& TemplateData = GetSharedTemplateData();
-	const TConstArrayView<FSceneStateTransition> TransitionsView = TemplateData.Transitions;
+	const TConstArrayView<FSceneStateTransition> TransitionsView = TemplateData->Transitions;
 
 	checkfSlow(TransitionsView.IsValidIndex(TransitionRange.Index) && TransitionsView.IsValidIndex(TransitionRange.GetLastIndex())
 		, TEXT("[Context: %s] State Transition Range [%d, %d] is out of bounds! (Total Transition Count: %d)")
@@ -308,8 +292,7 @@ TConstArrayView<FSceneStateTransition> FSceneStateExecutionContext::GetTransitio
 		return {};
 	}
 
-	const FSceneStateTemplateData& TemplateData = GetSharedTemplateData();
-	const TConstArrayView<FSceneStateTransition> TransitionsView = TemplateData.Transitions;
+	const TConstArrayView<FSceneStateTransition> TransitionsView = TemplateData->Transitions;
 
 	checkfSlow(TransitionsView.IsValidIndex(TransitionRange.Index) && TransitionsView.IsValidIndex(TransitionRange.GetLastIndex())
 		, TEXT("[Context: %s] Conduit Transition Range [%d, %d] is out of bounds! (Total Transition Count: %d)")
@@ -323,8 +306,7 @@ TConstArrayView<FSceneStateTransition> FSceneStateExecutionContext::GetTransitio
 
 FInstancedPropertyBag FSceneStateExecutionContext::GetTemplateTransitionParameter(const FSceneStateTransition& InTransition) const
 {
-	const FSceneStateTemplateData& TemplateData = GetSharedTemplateData();
-	return TemplateData.TransitionParameters.FindRef(GetTransitionIndex(TemplateData, InTransition));
+	return TemplateData->TransitionParameters.FindRef(GetTransitionIndex(InTransition));
 }
 
 TConstArrayView<FSceneStateMachine> FSceneStateExecutionContext::GetStateMachines(const FSceneState& InState) const
@@ -335,8 +317,7 @@ TConstArrayView<FSceneStateMachine> FSceneStateExecutionContext::GetStateMachine
 		return {};
 	}
 
-	const FSceneStateTemplateData& TemplateData = GetSharedTemplateData();
-	const TConstArrayView<FSceneStateMachine> StateMachinesView = TemplateData.StateMachines;
+	const TConstArrayView<FSceneStateMachine> StateMachinesView = TemplateData->StateMachines;
 
 	checkfSlow(StateMachinesView.IsValidIndex(StateMachineRange.Index) && StateMachinesView.IsValidIndex(StateMachineRange.GetLastIndex())
 		, TEXT("[Context: %s] State State Machine Range [%d, %d] is out of bounds! (Total State Machine Count: %d)")
@@ -356,8 +337,7 @@ TConstArrayView<FSceneStateEventHandler> FSceneStateExecutionContext::GetEventHa
 		return {};
 	}
 
-	const FSceneStateTemplateData& TemplateData = GetSharedTemplateData();
-	const TConstArrayView<FSceneStateEventHandler> EventHandlersView = TemplateData.EventHandlers;
+	const TConstArrayView<FSceneStateEventHandler> EventHandlersView = TemplateData->EventHandlers;
 
 	checkfSlow(EventHandlersView.IsValidIndex(EventHandlerRange.Index) && EventHandlersView.IsValidIndex(EventHandlerRange.GetLastIndex())
 		, TEXT("[Context: %s] State Event Handler Range [%d, %d] is out of bounds! (Total Event Handler Count: %d)")
@@ -372,10 +352,8 @@ TConstArrayView<FSceneStateEventHandler> FSceneStateExecutionContext::GetEventHa
 #if WITH_EDITOR
 const FSceneStateMetadata* FSceneStateExecutionContext::GetStateMetadata(const FSceneState& InState) const
 {
-	const FSceneStateTemplateData& TemplateData = GetSharedTemplateData();
-
-	const uint16 StateIndex = GetStateIndex(TemplateData, InState);
-	const TConstArrayView<FSceneStateMetadata> StateMetadataView = TemplateData.StateMetadata;
+	const uint16 StateIndex = GetStateIndex(InState);
+	const TConstArrayView<FSceneStateMetadata> StateMetadataView = TemplateData->StateMetadata;
 
 	if (!StateMetadataView.IsValidIndex(StateIndex))
 	{
@@ -393,12 +371,12 @@ const FSceneStateMetadata* FSceneStateExecutionContext::GetStateMetadata(const F
 
 const FSceneStateMachine* FSceneStateExecutionContext::GetStateMachine(const FGuid& InStateMachineId) const
 {
-	return GetSharedTemplateData().FindStateMachine(InStateMachineId);
+	return TemplateData->FindStateMachine(InStateMachineId);
 }
 
 const FSceneStateMachine* FSceneStateExecutionContext::GetStateMachine(uint16 InAbsoluteIndex) const
 {
-	const TConstArrayView<FSceneStateMachine> StateMachinesView = GetSharedTemplateData().StateMachines;
+	const TConstArrayView<FSceneStateMachine> StateMachinesView = TemplateData->StateMachines;
 
 	checkfSlow(StateMachinesView.IsValidIndex(InAbsoluteIndex)
 		, TEXT("[Context: %s] Absolute Index %d is out of bounds! (Total State Machine Count: %d)")
@@ -429,7 +407,7 @@ const FSceneState* FSceneStateExecutionContext::GetActiveState(const FSceneState
 
 const FSceneState* FSceneStateExecutionContext::GetState(uint16 InAbsoluteIndex) const
 {
-	const TConstArrayView<FSceneState> StatesView = GetSharedTemplateData().States;
+	const TConstArrayView<FSceneState> StatesView = TemplateData->States;
 
 	checkfSlow(StatesView.IsValidIndex(InAbsoluteIndex)
 		, TEXT("[Context: %s] Absolute Index %d is out of bounds! (Total State Count: %d)")
@@ -442,7 +420,7 @@ const FSceneState* FSceneStateExecutionContext::GetState(uint16 InAbsoluteIndex)
 
 const FSceneStateEventHandler* FSceneStateExecutionContext::GetEventHandler(uint16 InAbsoluteIndex) const
 {
-	const TConstArrayView<FSceneStateEventHandler> EventHandlers = GetSharedTemplateData().EventHandlers;
+	const TConstArrayView<FSceneStateEventHandler> EventHandlers = TemplateData->EventHandlers;
 
 	checkfSlow(EventHandlers.IsValidIndex(InAbsoluteIndex)
 		, TEXT("[Context: %s] Absolute Index %d is out of bounds! (Total Event Handlers Count: %d)")
@@ -461,7 +439,7 @@ const FSceneState* FSceneStateExecutionContext::GetState(const FSceneStateMachin
 		return nullptr;
 	}
 
-	const TConstArrayView<FSceneState> StatesView = GetSharedTemplateData().States;
+	const TConstArrayView<FSceneState> StatesView = TemplateData->States;
 
 	const uint16 AbsoluteIndex = StateRange.Index + InRelativeIndex;
 
@@ -484,7 +462,7 @@ const FSceneStateConduit* FSceneStateExecutionContext::GetConduit(const FSceneSt
 		return nullptr;
 	}
 
-	const TConstArrayView<FSceneStateConduit> ConduitsView = GetSharedTemplateData().Conduits;
+	const TConstArrayView<FSceneStateConduit> ConduitsView = TemplateData->Conduits;
 
 	const uint16 AbsoluteIndex = ConduitRange.Index + InRelativeIndex;
 
@@ -501,8 +479,7 @@ const FSceneStateConduit* FSceneStateExecutionContext::GetConduit(const FSceneSt
 
 FSceneStateInstance* FSceneStateExecutionContext::FindOrAddStateInstance(const FSceneState& InState) const
 {
-	const FSceneStateTemplateData& TemplateData = GetSharedTemplateData();
-	return &StateInstances.FindOrAdd(GetStateIndex(TemplateData, InState));
+	return &StateInstances.FindOrAdd(GetStateIndex(InState));
 }
 
 FSceneStateInstance* FSceneStateExecutionContext::FindStateInstance(uint16 InAbsoluteIndex) const
@@ -512,22 +489,19 @@ FSceneStateInstance* FSceneStateExecutionContext::FindStateInstance(uint16 InAbs
 
 FSceneStateInstance* FSceneStateExecutionContext::FindStateInstance(const FSceneState& InState) const
 {
-	const FSceneStateTemplateData& TemplateData = GetSharedTemplateData();
-	return FindStateInstance(GetStateIndex(TemplateData, InState));
+	return FindStateInstance(GetStateIndex(InState));
 }
 
 void FSceneStateExecutionContext::RemoveStateInstance(const FSceneState& InState) const
 {
-	const FSceneStateTemplateData& TemplateData = GetSharedTemplateData();
-	StateInstances.Remove(GetStateIndex(TemplateData, InState));
+	StateInstances.Remove(GetStateIndex(InState));
 }
 
 FConstStructView FSceneStateExecutionContext::FindTask(uint16 InAbsoluteIndex) const
 {
-	const FInstancedStructContainer& Tasks = GetSharedTemplateData().Tasks;
-	if (Tasks.IsValidIndex(InAbsoluteIndex))
+	if (TemplateData->Tasks.IsValidIndex(InAbsoluteIndex))
 	{
-		return Tasks[InAbsoluteIndex];
+		return TemplateData->Tasks[InAbsoluteIndex];
 	}
 	return {};
 }
@@ -535,14 +509,12 @@ FConstStructView FSceneStateExecutionContext::FindTask(uint16 InAbsoluteIndex) c
 FInstancedStructContainer* FSceneStateExecutionContext::FindOrAddTaskInstanceContainer(const FSceneState& InState) const
 {
 	// Note: task instance containers are allocated per state
-	const FSceneStateTemplateData& TemplateData = GetSharedTemplateData();
-	return &TaskInstanceContainers.FindOrAdd(GetStateIndex(TemplateData, InState));
+	return &TaskInstanceContainers.FindOrAdd(GetStateIndex(InState));
 }
 
 FInstancedStructContainer* FSceneStateExecutionContext::FindTaskInstanceContainer(const FSceneState& InState) const
 {
-	const FSceneStateTemplateData& TemplateData = GetSharedTemplateData();
-	return FindTaskInstanceContainer(GetStateIndex(TemplateData, InState));
+	return FindTaskInstanceContainer(GetStateIndex(InState));
 }
 
 FInstancedStructContainer* FSceneStateExecutionContext::FindTaskInstanceContainer(uint16 InAbsoluteIndex) const
@@ -552,9 +524,7 @@ FInstancedStructContainer* FSceneStateExecutionContext::FindTaskInstanceContaine
 
 FStructView FSceneStateExecutionContext::FindTaskInstance(uint16 InAbsoluteIndex) const
 {
-	const FSceneStateTemplateData& TemplateData = GetSharedTemplateData();
-
-	const FInstancedStructContainer& Tasks = TemplateData.Tasks;
+	const FInstancedStructContainer& Tasks = TemplateData->Tasks;
 	if (!Tasks.IsValidIndex(InAbsoluteIndex))
 	{
 		return {};
@@ -570,7 +540,7 @@ FStructView FSceneStateExecutionContext::FindTaskInstance(uint16 InAbsoluteIndex
 		return {};
 	}
 
-	const TConstArrayView<FSceneState> StatesView = TemplateData.States;
+	const TConstArrayView<FSceneState> StatesView = TemplateData->States;
 
 	checkf(StatesView.IsValidIndex(StateIndex)
 		, TEXT("[Context: %s] Task Parent State Index %d is out of bounds! (Total State Count: %d)")
@@ -591,14 +561,12 @@ FStructView FSceneStateExecutionContext::FindTaskInstance(uint16 InAbsoluteIndex
 
 void FSceneStateExecutionContext::RemoveTaskInstanceContainer(const FSceneState& InState) const
 {
-	const FSceneStateTemplateData& TemplateData = GetSharedTemplateData();
-	TaskInstanceContainers.Remove(GetStateIndex(TemplateData, InState));
+	TaskInstanceContainers.Remove(GetStateIndex(InState));
 }
 
 FSceneStateMachineInstance* FSceneStateExecutionContext::FindOrAddStateMachineInstance(const FSceneStateMachine& InStateMachine) const
 {
-	const FSceneStateTemplateData& TemplateData = GetSharedTemplateData();
-	return &StateMachineInstances.FindOrAdd(GetStateMachineIndex(TemplateData, InStateMachine));
+	return &StateMachineInstances.FindOrAdd(GetStateMachineIndex(InStateMachine));
 }
 
 FSceneStateMachineInstance* FSceneStateExecutionContext::FindStateMachineInstance(uint16 InAbsoluteIndex) const
@@ -608,20 +576,17 @@ FSceneStateMachineInstance* FSceneStateExecutionContext::FindStateMachineInstanc
 
 FSceneStateMachineInstance* FSceneStateExecutionContext::FindStateMachineInstance(const FSceneStateMachine& InStateMachine) const
 {
-	const FSceneStateTemplateData& TemplateData = GetSharedTemplateData();
-	return FindStateMachineInstance(GetStateMachineIndex(TemplateData, InStateMachine));
+	return FindStateMachineInstance(GetStateMachineIndex(InStateMachine));
 }
 
 void FSceneStateExecutionContext::RemoveStateMachineInstance(const FSceneStateMachine& InStateMachine) const
 {
-	const FSceneStateTemplateData& TemplateData = GetSharedTemplateData();
-	StateMachineInstances.Remove(GetStateMachineIndex(TemplateData, InStateMachine));
+	StateMachineInstances.Remove(GetStateMachineIndex(InStateMachine));
 }
 
 FSceneStateTransitionInstance* FSceneStateExecutionContext::FindOrAddTransitionInstance(const FSceneStateTransition& InTransition) const
 {
-	const FSceneStateTemplateData& TemplateData = GetSharedTemplateData();
-	return &TransitionInstances.FindOrAdd(GetTransitionIndex(TemplateData, InTransition));
+	return &TransitionInstances.FindOrAdd(GetTransitionIndex(InTransition));
 }
 
 FSceneStateTransitionInstance* FSceneStateExecutionContext::FindTransitionInstance(uint16 InAbsoluteIndex) const
@@ -631,26 +596,22 @@ FSceneStateTransitionInstance* FSceneStateExecutionContext::FindTransitionInstan
 
 FSceneStateTransitionInstance* FSceneStateExecutionContext::FindTransitionInstance(const FSceneStateTransition& InTransition) const
 {
-	const FSceneStateTemplateData& TemplateData = GetSharedTemplateData();
-	return FindTransitionInstance(GetTransitionIndex(TemplateData, InTransition));
+	return FindTransitionInstance(GetTransitionIndex(InTransition));
 }
 
 void FSceneStateExecutionContext::RemoveTransitionInstance(const FSceneStateTransition& InTransition) const
 {
-	const FSceneStateTemplateData& TemplateData = GetSharedTemplateData();
-	TransitionInstances.Remove(GetTransitionIndex(TemplateData, InTransition));
+	TransitionInstances.Remove(GetTransitionIndex(InTransition));
 }
 
 FSceneStateConduitInstance* FSceneStateExecutionContext::FindOrAddConduitInstance(const FSceneStateConduit& InConduit) const
 {
-	const FSceneStateTemplateData& TemplateData = GetSharedTemplateData();
-	return &ConduitInstances.FindOrAdd(GetConduitIndex(TemplateData, InConduit));
+	return &ConduitInstances.FindOrAdd(GetConduitIndex(InConduit));
 }
 
 void FSceneStateExecutionContext::RemoveConduitInstance(const FSceneStateConduit& InConduit) const
 {
-	const FSceneStateTemplateData& TemplateData = GetSharedTemplateData();
-	ConduitInstances.Remove(GetConduitIndex(TemplateData, InConduit));
+	ConduitInstances.Remove(GetConduitIndex(InConduit));	
 }
 
 bool FSceneStateExecutionContext::SetupFunctionInstances(FPropertyBindingIndex16 InBindingsBatch) const
@@ -661,9 +622,7 @@ bool FSceneStateExecutionContext::SetupFunctionInstances(FPropertyBindingIndex16
 		return false;
 	}
 
-	const FSceneStateTemplateData& TemplateData = GetSharedTemplateData();
-
-	const FPropertyBindingCopyInfoBatch& CopyInfoBatch = TemplateData.BindingCollection.GetBatch(InBindingsBatch);
+	const FPropertyBindingCopyInfoBatch& CopyInfoBatch = TemplateData->BindingCollection.GetBatch(InBindingsBatch);
 
 	if (!UE::SceneState::IsValidRange(CopyInfoBatch.PropertyFunctionsBegin, CopyInfoBatch.PropertyFunctionsEnd))
 	{
@@ -681,7 +640,7 @@ bool FSceneStateExecutionContext::SetupFunctionInstances(FPropertyBindingIndex16
 
 		for (uint16 FunctionIndex = FunctionInstanceBatch.FunctionRange.Index; FunctionIndex <= FunctionInstanceBatch.FunctionRange.GetLastIndex(); ++FunctionIndex)
 		{
-			TemplateFunctionInstances.Add(TemplateData.FunctionInstances[FunctionIndex]);
+			TemplateFunctionInstances.Add(TemplateData->FunctionInstances[FunctionIndex]);
 		}
 
 		FunctionInstanceBatch.FunctionInstances.Append(TemplateFunctionInstances);
@@ -697,7 +656,7 @@ bool FSceneStateExecutionContext::SetupFunctionInstances(FPropertyBindingIndex16
 	// Set up these functions as they could be bound to other functions
 	for (uint16 FunctionIndex = FunctionInstanceBatch.FunctionRange.Index; FunctionIndex <= FunctionInstanceBatch.FunctionRange.GetLastIndex(); ++FunctionIndex)
 	{
-		if (const FSceneStateFunction* Function = TemplateData.Functions[FunctionIndex].GetPtr<const FSceneStateFunction>())
+		if (const FSceneStateFunction* Function = TemplateData->Functions[FunctionIndex].GetPtr<const FSceneStateFunction>())
 		{
 			Function->Setup(*this);
 		}
@@ -723,12 +682,10 @@ void FSceneStateExecutionContext::RemoveFunctionInstances(FPropertyBindingIndex1
 		return;
 	}
 
-	const FSceneStateTemplateData& TemplateData = GetSharedTemplateData();
-
 	// Clean up the functions used resources (e.g. other function instance data)
 	for (uint16 FunctionIndex = FunctionRange.Index; FunctionIndex <= FunctionRange.GetLastIndex(); ++FunctionIndex)
 	{
-		if (const FSceneStateFunction* Function = TemplateData.Functions[FunctionIndex].GetPtr<const FSceneStateFunction>())
+		if (const FSceneStateFunction* Function = TemplateData->Functions[FunctionIndex].GetPtr<const FSceneStateFunction>())
 		{
 			Function->Cleanup(*this);
 		}
@@ -744,7 +701,7 @@ void FSceneStateExecutionContext::RemoveFunctionInstances(FPropertyBindingIndex1
 
 FConstStructView FSceneStateExecutionContext::FindFunction(uint16 InAbsoluteIndex) const
 {
-	const FInstancedStructContainer& Functions = GetSharedTemplateData().Functions;
+	const FInstancedStructContainer& Functions = TemplateData->Functions;
 	if (Functions.IsValidIndex(InAbsoluteIndex))
 	{
 		return Functions[InAbsoluteIndex];
@@ -766,33 +723,33 @@ const FSceneStateExecutionContext* FSceneStateExecutionContext::FindParentContex
 	return nullptr;
 }
 
-uint16 FSceneStateExecutionContext::GetStateIndex(const FSceneStateTemplateData& InTemplateData, const FSceneState& InState) const
+uint16 FSceneStateExecutionContext::GetStateIndex(const FSceneState& InState) const
 {
-	TConstArrayView<FSceneState> StatesView = GetSharedTemplateData().States;
+	TConstArrayView<FSceneState> StatesView = TemplateData->States;
 	const int32 Index = static_cast<int32>(&InState - StatesView.GetData());
 	checkSlow(StatesView.IsValidIndex(Index) && Index >= 0 && Index <= TNumericLimits<uint16>::Max());
 	return static_cast<uint16>(Index);
 }
 
-uint16 FSceneStateExecutionContext::GetStateMachineIndex(const FSceneStateTemplateData& InTemplateData, const FSceneStateMachine& InStateMachine) const
+uint16 FSceneStateExecutionContext::GetStateMachineIndex(const FSceneStateMachine& InStateMachine) const
 {
-	TConstArrayView<FSceneStateMachine> StateMachinesView = InTemplateData.StateMachines;
+	TConstArrayView<FSceneStateMachine> StateMachinesView = TemplateData->StateMachines;
 	const int32 Index = static_cast<int32>(&InStateMachine - StateMachinesView.GetData());
 	checkSlow(StateMachinesView.IsValidIndex(Index) && Index >= 0 && Index <= TNumericLimits<uint16>::Max());
 	return static_cast<uint16>(Index);
 }
 
-uint16 FSceneStateExecutionContext::GetTransitionIndex(const FSceneStateTemplateData& InTemplateData, const FSceneStateTransition& InTransition) const
+uint16 FSceneStateExecutionContext::GetTransitionIndex(const FSceneStateTransition& InTransition) const
 {
-	TConstArrayView<FSceneStateTransition> TransitionsView = InTemplateData.Transitions;
+	TConstArrayView<FSceneStateTransition> TransitionsView = TemplateData->Transitions;
 	const int32 Index = static_cast<int32>(&InTransition - TransitionsView.GetData());
 	checkSlow(TransitionsView.IsValidIndex(Index) && Index >= 0 && Index <= TNumericLimits<uint16>::Max());
 	return static_cast<uint16>(Index);
 }
 
-uint16 FSceneStateExecutionContext::GetConduitIndex(const FSceneStateTemplateData& InTemplateData, const FSceneStateConduit& InConduit) const
+uint16 FSceneStateExecutionContext::GetConduitIndex(const FSceneStateConduit& InConduit) const
 {
-	TConstArrayView<FSceneStateConduit> ConduitsView = InTemplateData.Conduits;
+	TConstArrayView<FSceneStateConduit> ConduitsView = TemplateData->Conduits;
 	const int32 Index = static_cast<int32>(&InConduit - ConduitsView.GetData());
 	checkSlow(ConduitsView.IsValidIndex(Index) && Index >= 0 && Index <= TNumericLimits<uint16>::Max());
 	return static_cast<uint16>(Index);
diff --git a/Engine/Plugins/VirtualProduction/SceneState/Source/SceneState/Private/SceneStateExecutionScope.cpp b/Engine/Plugins/VirtualProduction/SceneState/Source/SceneState/Private/SceneStateExecutionScope.cpp
--- a/Engine/Plugins/VirtualProduction/SceneState/Source/SceneState/Private/SceneStateExecutionScope.cpp
+++ b/Engine/Plugins/VirtualProduction/SceneState/Source/SceneState/Private/SceneStateExecutionScope.cpp
@@ -9,8 +9,7 @@ namespace UE::SceneState
 FExecutionScope::FExecutionScope(const FSceneStateExecutionContext& InContext, const FSceneStateMachine& InStateMachine)
 	: ExecutionContext(InContext)
 {
-	const FSceneStateTemplateData& TemplateData = InContext.GetSharedTemplateData();
-	StateMachineIndex = ExecutionContext.GetStateMachineIndex(TemplateData, InStateMachine);
+	StateMachineIndex = ExecutionContext.GetStateMachineIndex(InStateMachine);
 	ExecutionContext.StateMachineExecutionStack.Push(StateMachineIndex);
 }
 
diff --git a/Engine/Plugins/VirtualProduction/SceneState/Source/SceneState/Private/SceneStateGeneratedClass.cpp b/Engine/Plugins/VirtualProduction/SceneState/Source/SceneState/Private/SceneStateGeneratedClass.cpp
--- a/Engine/Plugins/VirtualProduction/SceneState/Source/SceneState/Private/SceneStateGeneratedClass.cpp
+++ b/Engine/Plugins/VirtualProduction/SceneState/Source/SceneState/Private/SceneStateGeneratedClass.cpp
@@ -10,8 +10,6 @@
 
 USceneStateGeneratedClass::USceneStateGeneratedClass()
 {
-	SharedTemplateData.InitializeAs<FSceneStateTemplateData>(TNotNull<const UClass*>(this));
-
 #if WITH_EDITOR
 	if (!HasAnyFlags(RF_ClassDefaultObject))
 	{
@@ -21,24 +19,23 @@ USceneStateGeneratedClass::USceneStateGeneratedClass()
 #endif
 }
 
-FConstSharedStruct USceneStateGeneratedClass::GetSharedTemplateData() const
-{
-	return SharedTemplateData;
-}
-
 void USceneStateGeneratedClass::Link(FArchive& Ar, bool bInRelinkExistingProperties)
 {
 	Super::Link(Ar, bInRelinkExistingProperties);
-	if (IsFullClass())
+
+	if (TemplateData && IsFullClass())
 	{
-		GetTemplateData().Link();
+		TemplateData->Link(this);
 	}
 }
 
 void USceneStateGeneratedClass::PurgeClass(bool bInRecompilingOnLoad)
 {
 	Super::PurgeClass(bInRecompilingOnLoad);
-	GetTemplateData().Reset();
+	if (TemplateData)
+	{
+		TemplateData->Reset();
+	}
 }
 
 void USceneStateGeneratedClass::PostLoad()
@@ -50,20 +47,17 @@ void USceneStateGeneratedClass::PostLoad()
 void USceneStateGeneratedClass::BeginDestroy()
 {
 	Super::BeginDestroy();
-	GetTemplateData().Reset();
-	SharedTemplateData.Reset();
-}
-
-FSceneStateTemplateData& USceneStateGeneratedClass::GetTemplateData()
-{
-	return SharedTemplateData.Get<FSceneStateTemplateData>();
+	if (TemplateData)
+	{
+		TemplateData->Reset();
+	}
 }
 
 void USceneStateGeneratedClass::ResolveBindings()
 {
-	if (IsFullClass())
+	if (TemplateData && IsFullClass())
 	{
-		GetTemplateData().ResolveBindings();
+		TemplateData->ResolveBindings(this);
 	}
 }
 
@@ -81,17 +75,17 @@ bool USceneStateGeneratedClass::IsFullClass() const
 #if WITH_EDITOR
 void USceneStateGeneratedClass::OnObjectsReinstanced(const TMap<UObject*, UObject*>& InReplacementMap)
 {
-	if (IsFullClass())
+	if (TemplateData && IsFullClass())
 	{
-		GetTemplateData().OnObjectsReinstanced(InReplacementMap);
+		TemplateData->OnObjectsReinstanced(this, InReplacementMap);
 	}
 }
 
 void USceneStateGeneratedClass::OnStructsReinstanced(const UUserDefinedStruct& InStruct)
 {
-	if (IsFullClass())
+	if (TemplateData && IsFullClass())
 	{
-		GetTemplateData().OnStructsReinstanced(InStruct);
+		TemplateData->OnStructsReinstanced(this, InStruct);
 	}
 }
 #endif
diff --git a/Engine/Plugins/VirtualProduction/SceneState/Source/SceneState/Private/SceneStateObject.cpp b/Engine/Plugins/VirtualProduction/SceneState/Source/SceneState/Private/SceneStateObject.cpp
--- a/Engine/Plugins/VirtualProduction/SceneState/Source/SceneState/Private/SceneStateObject.cpp
+++ b/Engine/Plugins/VirtualProduction/SceneState/Source/SceneState/Private/SceneStateObject.cpp
@@ -6,6 +6,7 @@
 #include "SceneStateEventStream.h"
 #include "SceneStateExecutionContextRegistry.h"
 #include "SceneStateGeneratedClass.h"
+#include "SceneStateLog.h"
 #include "SceneStatePlayer.h"
 #include "SceneStateTemplateData.h"
 
@@ -25,11 +26,6 @@ FString USceneStateObject::GetContextName() const
 	return FString();
 }
 
-FConstSharedStruct USceneStateObject::GetSharedTemplateData() const
-{
-	return SharedTemplateData;
-}
-
 UObject* USceneStateObject::GetContextObject() const
 {
 	if (USceneStatePlayer* Player = Cast<USceneStatePlayer>(GetOuter()))
@@ -51,22 +47,31 @@ bool USceneStateObject::IsActive() const
 
 void USceneStateObject::Setup()
 {
+	TemplateData = nullptr;
+
 	if (const USceneStateGeneratedClass* GeneratedClass = Cast<USceneStateGeneratedClass>(GetClass()))
 	{
-		SharedTemplateData = GeneratedClass->GetSharedTemplateData();
+		TemplateData = GeneratedClass->GetTemplateData();
+	}
+
+	if (!TemplateData)
+	{
+		UE_LOG(LogSceneState, Error, TEXT("[%s] did not find a valid template data object."), *GetFullName());
+		return;
 	}
+
 	RootExecutionContext.Setup(this);
 }
 
 void USceneStateObject::Enter()
 {
-	if (EventStream)
-	{
-		EventStream->Register();
-	}
-
 	if (const FSceneState* RootState = GetRootState())
 	{
+		if (EventStream)
+		{
+			EventStream->Register();
+		}
+
 		ReceiveEnter();
 		RootState->Enter(RootExecutionContext);
 	}
@@ -94,7 +99,7 @@ void USceneStateObject::Exit()
 		EventStream->Unregister();
 	}
 
-	SharedTemplateData.Reset();
+	TemplateData = nullptr;
 	RootExecutionContext.Reset();
 }
 
@@ -115,13 +120,13 @@ UWorld* USceneStateObject::GetWorld() const
 void USceneStateObject::BeginDestroy()
 {
 	Super::BeginDestroy();
-	SharedTemplateData.Reset();
+	TemplateData = nullptr;
     RootExecutionContext.Reset();
 }
 
 const FSceneState* USceneStateObject::GetRootState() const
 {
-	if (const FSceneStateTemplateData* TemplateData = GetSharedTemplateData().GetPtr<const FSceneStateTemplateData>())
+	if (TemplateData)
 	{
 		return TemplateData->GetRootState();
 	}
diff --git a/Engine/Plugins/VirtualProduction/SceneState/Source/SceneState/Private/SceneStateTemplateData.cpp b/Engine/Plugins/VirtualProduction/SceneState/Source/SceneState/Private/SceneStateTemplateData.cpp
--- a/Engine/Plugins/VirtualProduction/SceneState/Source/SceneState/Private/SceneStateTemplateData.cpp
+++ b/Engine/Plugins/VirtualProduction/SceneState/Source/SceneState/Private/SceneStateTemplateData.cpp
@@ -11,12 +11,7 @@
 #include "Tasks/SceneStateTaskBindingExtension.h"
 #include "Tasks/SceneStateTaskInstance.h"
 
-FSceneStateTemplateData::FSceneStateTemplateData(TNotNull<const UClass*> InOwnerClass)
-	: OwnerClass(InOwnerClass)
-{
-}
-
-const FSceneState* FSceneStateTemplateData::GetRootState() const
+const FSceneState* USceneStateTemplateData::GetRootState() const
 {
 	if (States.IsValidIndex(RootStateIndex))
 	{
@@ -25,7 +20,7 @@ const FSceneState* FSceneStateTemplateData::GetRootState() const
 	return nullptr;
 }
 
-const FSceneStateMachine* FSceneStateTemplateData::FindStateMachine(const FGuid& InStateMachineId) const
+const FSceneStateMachine* USceneStateTemplateData::FindStateMachine(const FGuid& InStateMachineId) const
 {
 	const uint16* StateMachineIndex = StateMachineIdToIndex.Find(InStateMachineId);
 	if (StateMachineIndex && StateMachines.IsValidIndex(*StateMachineIndex))
@@ -35,7 +30,7 @@ const FSceneStateMachine* FSceneStateTemplateData::FindStateMachine(const FGuid&
 	return nullptr;
 }
 
-const UStruct* FSceneStateTemplateData::FindDataStruct(const FSceneStateBindingDataHandle& InDataHandle)
+const UStruct* USceneStateTemplateData::FindDataStruct(const UClass* InOwnerClass, const FSceneStateBindingDataHandle& InDataHandle)
 {
 	const uint16 DataIndex = InDataHandle.GetDataIndex();
 	const uint16 DataSubIndex = InDataHandle.GetDataSubIndex();
@@ -51,7 +46,7 @@ const UStruct* FSceneStateTemplateData::FindDataStruct(const FSceneStateBindingD
 	switch (DataType)
 	{
 	case ESceneStateDataType::Root:
-		return OwnerClass;
+		return InOwnerClass;
 
 	case ESceneStateDataType::Task:
 		return TaskInstances[DataIndex].GetScriptStruct();
@@ -86,14 +81,14 @@ const UStruct* FSceneStateTemplateData::FindDataStruct(const FSceneStateBindingD
 	return nullptr;
 }
 
-void FSceneStateTemplateData::ResolveBindings()
+void USceneStateTemplateData::ResolveBindings(const UClass* InOwnerClass)
 {
 	UE::SceneState::PatchBindingCollection(
 		{
 			.BindingCollection = BindingCollection,
-			.FindDataStructFunctor = [this](const FSceneStateBindingDataHandle& InDataHandle)->const UStruct*
+			.FindDataStructFunctor = [This=this, InOwnerClass](const FSceneStateBindingDataHandle& InDataHandle)->const UStruct*
 			{
-				return FindDataStruct(InDataHandle);
+				return This->FindDataStruct(InOwnerClass, InDataHandle);
 			}
 		});
 
@@ -104,13 +99,13 @@ void FSceneStateTemplateData::ResolveBindings()
 	}
 }
 
-void FSceneStateTemplateData::Link()
+void USceneStateTemplateData::Link(const UClass* InOwnerClass)
 {
 	for (TEnumerateRef<FSceneStateTransition> Transition : EnumerateRange(Transitions))
 	{
 		if (TransitionLinks.IsValidIndex(Transition.GetIndex()))
 		{
-			Transition->Link(TransitionLinks[Transition.GetIndex()], OwnerClass);
+			Transition->Link(TransitionLinks[Transition.GetIndex()], InOwnerClass);
 		}
 	}
 
@@ -118,12 +113,12 @@ void FSceneStateTemplateData::Link()
 	{
 		if (ConduitLinks.IsValidIndex(Conduit.GetIndex()))
 		{
-			Conduit->Link(ConduitLinks[Conduit.GetIndex()], OwnerClass);
+			Conduit->Link(ConduitLinks[Conduit.GetIndex()], InOwnerClass);
 		}
 	}
 }
 
-void FSceneStateTemplateData::Reset()
+void USceneStateTemplateData::Reset()
 {
 	Tasks.Reset();
 	TaskPrerequisites.Reset();
@@ -153,7 +148,7 @@ void FSceneStateTemplateData::Reset()
 }
 
 #if WITH_EDITOR
-const FSceneState* FSceneStateTemplateData::FindStateFromNode(FObjectKey InStateNode) const
+const FSceneState* USceneStateTemplateData::FindStateFromNode(FObjectKey InStateNode) const
 {
 	const uint16* StateIndex = StateNodeToIndex.Find(InStateNode);
 	if (StateIndex && States.IsValidIndex(*StateIndex))
@@ -163,7 +158,7 @@ const FSceneState* FSceneStateTemplateData::FindStateFromNode(FObjectKey InState
 	return nullptr;
 }
 
-const FSceneStateTask* FSceneStateTemplateData::FindTaskFromNode(FObjectKey InTaskNode) const
+const FSceneStateTask* USceneStateTemplateData::FindTaskFromNode(FObjectKey InTaskNode) const
 {
 	const uint16* TaskIndex = TaskNodeToIndex.Find(InTaskNode);
 	if (TaskIndex && Tasks.IsValidIndex(*TaskIndex))
@@ -173,7 +168,7 @@ const FSceneStateTask* FSceneStateTemplateData::FindTaskFromNode(FObjectKey InTa
 	return nullptr;
 }
 
-void FSceneStateTemplateData::ForEachStateInstance(const TSharedRef<UE::SceneState::FExecutionContextRegistry>& InContextRegistry
+void USceneStateTemplateData::ForEachStateInstance(const TSharedRef<UE::SceneState::FExecutionContextRegistry>& InContextRegistry
 	, FObjectKey InStateNode
 	, TFunctionRef<void(const FSceneStateInstance&)> InFunctor) const
 {
@@ -193,7 +188,7 @@ void FSceneStateTemplateData::ForEachStateInstance(const TSharedRef<UE::SceneSta
 		});
 }
 
-void FSceneStateTemplateData::ForEachTaskInstance(const TSharedRef<UE::SceneState::FExecutionContextRegistry>& InContextRegistry, FObjectKey InTaskNode, TFunctionRef<void(const FSceneStateTaskInstance&)> InFunctor) const
+void USceneStateTemplateData::ForEachTaskInstance(const TSharedRef<UE::SceneState::FExecutionContextRegistry>& InContextRegistry, FObjectKey InTaskNode, TFunctionRef<void(const FSceneStateTaskInstance&)> InFunctor) const
 {
 	const uint16* TaskIndex = TaskNodeToIndex.Find(InTaskNode);
 	if (!TaskIndex)
@@ -211,7 +206,7 @@ void FSceneStateTemplateData::ForEachTaskInstance(const TSharedRef<UE::SceneStat
 		});
 }
 
-void FSceneStateTemplateData::OnObjectsReinstanced(const TMap<UObject*, UObject*>& InReplacementMap)
+void USceneStateTemplateData::OnObjectsReinstanced(const UClass* InOwnerClass, const TMap<UObject*, UObject*>& InReplacementMap)
 {
 	TSet<const UStruct*> Structs;
 	Structs.Reserve(InReplacementMap.Num());
@@ -222,7 +217,7 @@ void FSceneStateTemplateData::OnObjectsReinstanced(const TMap<UObject*, UObject*
 	{
 		if (const UObject* Replacement = Pair.Value)
 		{
-			if (Replacement->IsIn(OwnerClass))
+			if (Replacement->IsIn(InOwnerClass))
 			{
 				bRequiresResolve = true;
 				break;
@@ -233,18 +228,18 @@ void FSceneStateTemplateData::OnObjectsReinstanced(const TMap<UObject*, UObject*
 
 	if (bRequiresResolve || BindingCollection.ContainsAnyStruct(Structs))
 	{
-		ResolveBindings();
+		ResolveBindings(InOwnerClass);
 	}
 }
 
-void FSceneStateTemplateData::OnStructsReinstanced(const UUserDefinedStruct& InStruct)
+void USceneStateTemplateData::OnStructsReinstanced(const UClass* InOwnerClass, const UUserDefinedStruct& InStruct)
 {
 	TSet<const UStruct*> Structs;
 	Structs.Add(&InStruct);
 
 	if (BindingCollection.ContainsAnyStruct(Structs))
 	{
-		ResolveBindings();
+		ResolveBindings(InOwnerClass);
 	}
 }
 #endif
diff --git a/Engine/Plugins/VirtualProduction/SceneState/Source/SceneState/Public/SceneStateExecutionContext.h b/Engine/Plugins/VirtualProduction/SceneState/Source/SceneState/Public/SceneStateExecutionContext.h
--- a/Engine/Plugins/VirtualProduction/SceneState/Source/SceneState/Public/SceneStateExecutionContext.h
+++ b/Engine/Plugins/VirtualProduction/SceneState/Source/SceneState/Public/SceneStateExecutionContext.h
@@ -20,6 +20,7 @@ class UClass;
 class USceneStateEventStream;
 class USceneStateObject;
 class USceneStatePlayer;
+class USceneStateTemplateData;
 struct FPropertyBindingDataView;
 struct FPropertyBindingIndex16;
 struct FSceneState;
@@ -32,7 +33,6 @@ struct FSceneStateMachine;
 struct FSceneStateMetadata;
 struct FSceneStatePropertyReference;
 struct FSceneStateTask;
-struct FSceneStateTemplateData;
 struct FSceneStateTransition;
 
 namespace UE::SceneState
@@ -51,7 +51,11 @@ namespace UE::SceneState
  * is immutable during execution.
  *
  * This struct also offers functionality to both get class (immutable) objects (states, state machines, etc)
- * and their instance (mutable) data (state instances, state machine instances, etc)
+ * and their instance (mutable) data (state instances, state machine instances, etc).
+ *
+ * Prior to execution, ensure that the template data is in a valid state,
+ * as it is assumed to be valid during execution to avoid re-checking validity.
+ * See USceneStateObject as an example.
  */
 USTRUCT()
 struct FSceneStateExecutionContext
@@ -62,12 +66,15 @@ struct FSceneStateExecutionContext
 
 	UE_API ~FSceneStateExecutionContext();
 
-	UE_API const FSceneStateTemplateData& GetSharedTemplateData() const;
-
 	/** Gets the registry that this context is registered to */
 	TSharedPtr<const UE::SceneState::FExecutionContextRegistry> GetContextRegistry() const;
 
-	/** Gets the binding collection owned by the generated class */
+	const USceneStateTemplateData* GetTemplateData() const
+	{
+		return TemplateData;
+	}
+
+	/** Gets the binding collection owned by the template data */
 	const FSceneStateBindingCollection& GetBindingCollection() const;
 
 	/** Returns the Player Debug Name for logging purposes */
@@ -98,7 +105,7 @@ struct FSceneStateExecutionContext
 	struct FSetupParams
 	{
 		/** Shared memory of the template data holding all the state machines, states, etc. */
-		FConstSharedStruct SharedTemplateData;
+		const USceneStateTemplateData* TemplateData;
 		/** Root object instance */
 		UObject* RootObject = nullptr;
 		/** Context object for utility purposes (e.g. to get world from context, etc) */
@@ -263,16 +270,16 @@ private:
 	const FSceneStateExecutionContext* FindParentContext() const;
 
 	/** Retrieves the Absolute Index of a State in the State array */
-	uint16 GetStateIndex(const FSceneStateTemplateData& InTemplateData, const FSceneState& InState) const;
+	uint16 GetStateIndex(const FSceneState& InState) const;
 
 	/** Retrieves the Absolute Index of a State Machine in the State Machine array */
-	uint16 GetStateMachineIndex(const FSceneStateTemplateData& InTemplateData, const FSceneStateMachine& InStateMachine) const;
+	uint16 GetStateMachineIndex(const FSceneStateMachine& InStateMachine) const;
 
 	/** Retrieves the Absolute Index of a Transition in the Transitions array */
-	uint16 GetTransitionIndex(const FSceneStateTemplateData& InTemplateData, const FSceneStateTransition& InTransition) const;
+	uint16 GetTransitionIndex(const FSceneStateTransition& InTransition) const;
 
 	/** Retrieves the Absolute Index of a Conduit in the Conduits array */
-	uint16 GetConduitIndex(const FSceneStateTemplateData& InTemplateData, const FSceneStateConduit& InConduit) const;
+	uint16 GetConduitIndex(const FSceneStateConduit& InConduit) const;
 
 	/** Weak ptr to the context holding this execution */
 	UPROPERTY()
@@ -290,9 +297,9 @@ private:
 	UPROPERTY()
 	FString ContextName;
 
-	/** Scene State Template */
+	/** Object owning the scene state template data */
 	UPROPERTY()
-	FConstSharedStruct SharedTemplateData;
+	TObjectPtr<const USceneStateTemplateData> TemplateData;
 
 	/**
 	 * Map of State Index to its Instance Data
diff --git a/Engine/Plugins/VirtualProduction/SceneState/Source/SceneState/Public/SceneStateGeneratedClass.h b/Engine/Plugins/VirtualProduction/SceneState/Source/SceneState/Public/SceneStateGeneratedClass.h
--- a/Engine/Plugins/VirtualProduction/SceneState/Source/SceneState/Public/SceneStateGeneratedClass.h
+++ b/Engine/Plugins/VirtualProduction/SceneState/Source/SceneState/Public/SceneStateGeneratedClass.h
@@ -3,13 +3,12 @@
 #pragma once
 
 #include "Engine/BlueprintGeneratedClass.h"
-#include "StructUtils/SharedStruct.h"
 #include "SceneStateGeneratedClass.generated.h"
 
 #define UE_API SCENESTATE_API
 
+class USceneStateTemplateData;
 class UUserDefinedStruct;
-struct FSceneStateTemplateData;
 
 namespace UE::SceneState::Editor
 {
@@ -22,7 +21,7 @@ namespace UE::SceneState::Editor
 /**
  * Object Class for the Scene State Object
  * Holds the template data for scene state objects.
- * @see FSceneStateTemplateData
+ * @see USceneStateTemplateData
  * @see FSceneStateExecutionContext
  */
 UCLASS(MinimalAPI)
@@ -33,7 +32,10 @@ class USceneStateGeneratedClass : public UBlueprintGeneratedClass
 public:
 	UE_API USceneStateGeneratedClass();
 
-	UE_API FConstSharedStruct GetSharedTemplateData() const;
+	const USceneStateTemplateData* GetTemplateData() const
+	{
+		return TemplateData;
+	}
 
 	//~ Begin UClass
 	UE_API virtual void Link(FArchive& Ar, bool bInRelinkExistingProperties) override;
@@ -46,8 +48,6 @@ public:
 	//~ End UObject
 
 private:
-	FSceneStateTemplateData& GetTemplateData();
-
 	/** Called to resolve the template data bindings */
 	UE_API void ResolveBindings();
 
@@ -58,9 +58,9 @@ private:
 	void OnStructsReinstanced(const UUserDefinedStruct& InStruct);
 #endif
 
-	/** Scene State template data (FSceneStateTemplateData) for this class. */
+	/** Scene State template data for this class. */
 	UPROPERTY()
-	FSharedStruct SharedTemplateData;
+	TObjectPtr<USceneStateTemplateData> TemplateData;
 
 #if WITH_EDITOR
 	FDelegateHandle OnObjectsReinstancedHandle;
diff --git a/Engine/Plugins/VirtualProduction/SceneState/Source/SceneState/Public/SceneStateObject.h b/Engine/Plugins/VirtualProduction/SceneState/Source/SceneState/Public/SceneStateObject.h
--- a/Engine/Plugins/VirtualProduction/SceneState/Source/SceneState/Public/SceneStateObject.h
+++ b/Engine/Plugins/VirtualProduction/SceneState/Source/SceneState/Public/SceneStateObject.h
@@ -15,6 +15,7 @@
 
 class USceneStateBlueprintableTask;
 class USceneStateEventStream;
+class USceneStateTemplateData;
 struct FPropertyBindingDataView;
 struct FSceneStateBindingDataHandle;
 struct FSceneStateExecutionContext;
@@ -40,8 +41,10 @@ public:
 	/** Gets the context name of the scene state player */
 	UE_API FString GetContextName() const;
 
-	/** Gets the template data (FSceneStateTemplate) for the given object */
-	UE_API FConstSharedStruct GetSharedTemplateData() const;
+	const USceneStateTemplateData* GetTemplateData() const
+	{
+		return TemplateData;
+	}
 
 	/** Gets the context object of the scene state player */
 	UFUNCTION(BlueprintCallable, Category="State")
@@ -87,12 +90,12 @@ public:
 	//~ End UObject
 
 private:
-	/** Resolves the root state from the shared template data, if any */
+	/** Resolves the root state from the template data, if any */
 	const FSceneState* GetRootState() const;
 
-	/** Scene State template data (FSceneStateTemplateData) to use for this object's execution. */
+	/** Scene State template data to use for this object's execution. */
 	UPROPERTY(Transient, DuplicateTransient, TextExportTransient, NonTransactional, meta=(NoBinding))
-	FConstSharedStruct SharedTemplateData;
+	TObjectPtr<const USceneStateTemplateData> TemplateData;
 
 	/** Runtime Event System keeping track of Active Events that have been added */
 	UPROPERTY(Transient, Instanced, meta=(NoBinding))
diff --git a/Engine/Plugins/VirtualProduction/SceneState/Source/SceneState/Public/SceneStateTemplateData.h b/Engine/Plugins/VirtualProduction/SceneState/Source/SceneState/Public/SceneStateTemplateData.h
--- a/Engine/Plugins/VirtualProduction/SceneState/Source/SceneState/Public/SceneStateTemplateData.h
+++ b/Engine/Plugins/VirtualProduction/SceneState/Source/SceneState/Public/SceneStateTemplateData.h
@@ -14,6 +14,7 @@
 #include "StructUtils/PropertyBag.h"
 #include "Transition/SceneStateTransition.h"
 #include "Transition/SceneStateTransitionLink.h"
+#include "UObject/Object.h"
 #include "UObject/ObjectKey.h"
 
 #if WITH_EDITOR
@@ -40,19 +41,12 @@ struct FSceneStateTaskInstance;
  * All this data is immutable in execution, and as such, it is not instanced to the Scene State Object instances.
  * @see FSceneStateExecutionContext
  */
-USTRUCT()
-struct FSceneStateTemplateData
+UCLASS(MinimalAPI)
+class USceneStateTemplateData : public UObject
 {
 	GENERATED_BODY()
 
-	/** For compatibility only. Do not use */
-	explicit FSceneStateTemplateData(EForceInit)
-	{
-	}
-
-	/** Constructs the template data for a given owner class */
-	explicit FSceneStateTemplateData(TNotNull<const UClass*> InOwnerClass);
-
+public:
 	/** Gets the state that contains all top-level state machines */
 	const FSceneState* GetRootState() const;
 
@@ -60,13 +54,13 @@ struct FSceneStateTemplateData
 	const FSceneStateMachine* FindStateMachine(const FGuid& InStateMachineId) const;
 
 	/** Finds the struct type for the given data handle */
-	const UStruct* FindDataStruct(const FSceneStateBindingDataHandle& InDataHandle);
+	const UStruct* FindDataStruct(const UClass* InOwnerClass, const FSceneStateBindingDataHandle& InDataHandle);
 
 	/** Patches bindings and resolves the binding paths for the owning binding collection */
-	UE_API void ResolveBindings();
+	void ResolveBindings(const UClass* InOwnerClass);
 
 	/** Called at link time to cache the properties and functions to use (e.g. in transitions or conduits) */
-	UE_API void Link();
+	void Link(const UClass* InOwnerClass);
 
 	/** Resets all the elements that get compiled for this generated class (e.g. Tasks, State Machines, etc) */
 	UE_API void Reset();
@@ -98,14 +92,10 @@ struct FSceneStateTemplateData
 		, FObjectKey InTaskNode
 		, TFunctionRef<void(const FSceneStateTaskInstance&)> InFunctor) const;
 
-	void OnObjectsReinstanced(const TMap<UObject*, UObject*>& InReplacementMap);
-	void OnStructsReinstanced(const UUserDefinedStruct& InStruct);
+	void OnObjectsReinstanced(const UClass* InOwnerClass, const TMap<UObject*, UObject*>& InReplacementMap);
+	void OnStructsReinstanced(const UClass* InOwnerClass, const UUserDefinedStruct& InStruct);
 #endif
 
-	/** Owner class of this template data */
-	UPROPERTY()
-	TObjectPtr<const UClass> OwnerClass;
-
 	/** Container for all the bindings */
 	UPROPERTY()
 	FSceneStateBindingCollection BindingCollection;
@@ -197,13 +187,4 @@ struct FSceneStateTemplateData
 #endif
 };
 
-template<>
-struct TStructOpsTypeTraits<FSceneStateTemplateData> : public TStructOpsTypeTraitsBase2<FSceneStateTemplateData>
-{
-	enum
-	{
-		WithNoInitConstructor = true
-	};
-};
-
 #undef UE_API
diff --git a/Engine/Plugins/VirtualProduction/SceneState/Source/SceneStateBlueprintEditor/Private/Compiler/SceneStateBindingCompiler.cpp b/Engine/Plugins/VirtualProduction/SceneState/Source/SceneStateBlueprintEditor/Private/Compiler/SceneStateBindingCompiler.cpp
--- a/Engine/Plugins/VirtualProduction/SceneState/Source/SceneStateBlueprintEditor/Private/Compiler/SceneStateBindingCompiler.cpp
+++ b/Engine/Plugins/VirtualProduction/SceneState/Source/SceneStateBlueprintEditor/Private/Compiler/SceneStateBindingCompiler.cpp
@@ -60,7 +60,7 @@ EDataAccessType GetBindingBatchingType(const FSceneStateBindingCollection& InBin
 
 } // UE::SceneState::Editor::Private
 
-FBindingCompiler::FBindingCompiler(FBlueprintCompilerContext& InContext, TNotNull<USceneStateBlueprint*> InBlueprint, TNotNull<FSceneStateTemplateData*> InTemplateData)
+FBindingCompiler::FBindingCompiler(FBlueprintCompilerContext& InContext, TNotNull<USceneStateBlueprint*> InBlueprint, TNotNull<USceneStateTemplateData*> InTemplateData)
 	: Context(InContext)
 	, Blueprint(InBlueprint)
 	, TemplateData(InTemplateData)
@@ -116,7 +116,7 @@ void FBindingCompiler::AddRootBindingDesc()
 	FSceneStateBindingDesc BindingDesc = Blueprint->CreateRootBinding();
 
 	// Set the Struct to the New Class during compilation
-	BindingDesc.Struct = TemplateData->OwnerClass;
+	BindingDesc.Struct = Context.GetGeneratedClass();
 	BindingDesc.DataHandle = FSceneStateBindingDataHandle(ESceneStateDataType::Root);
 
 	AddBindingDesc(MoveTemp(BindingDesc));
diff --git a/Engine/Plugins/VirtualProduction/SceneState/Source/SceneStateBlueprintEditor/Private/Compiler/SceneStateBindingCompiler.h b/Engine/Plugins/VirtualProduction/SceneState/Source/SceneStateBlueprintEditor/Private/Compiler/SceneStateBindingCompiler.h
--- a/Engine/Plugins/VirtualProduction/SceneState/Source/SceneStateBlueprintEditor/Private/Compiler/SceneStateBindingCompiler.h
+++ b/Engine/Plugins/VirtualProduction/SceneState/Source/SceneStateBlueprintEditor/Private/Compiler/SceneStateBindingCompiler.h
@@ -9,6 +9,7 @@
 #include "UObject/ObjectPtr.h"
 
 class USceneStateBlueprint;
+class USceneStateTemplateData;
 struct FPropertyBindingIndex16;
 struct FSceneStateBinding;
 struct FSceneStateBindingDesc;
@@ -32,7 +33,7 @@ enum class EDataAccessType : uint8
 class FBindingCompiler
 {
 public:
-	explicit FBindingCompiler(FBlueprintCompilerContext& InContext, TNotNull<USceneStateBlueprint*> InBlueprint, TNotNull<FSceneStateTemplateData*> InTemplateData);
+	explicit FBindingCompiler(FBlueprintCompilerContext& InContext, TNotNull<USceneStateBlueprint*> InBlueprint, TNotNull<USceneStateTemplateData*> InTemplateData);
 
 	void Compile();
 
@@ -67,7 +68,7 @@ private:
 	TObjectPtr<USceneStateBlueprint> Blueprint;
 
 	/** The target to store the compiled bindings */
-	FSceneStateTemplateData* TemplateData;
+	TObjectPtr<USceneStateTemplateData> TemplateData;
 
 	/** Compiler for the functions found within the bindings */
 	FBindingFunctionCompiler BindingFunctionCompiler;
diff --git a/Engine/Plugins/VirtualProduction/SceneState/Source/SceneStateBlueprintEditor/Private/Compiler/SceneStateBindingFunctionCompiler.h b/Engine/Plugins/VirtualProduction/SceneState/Source/SceneStateBlueprintEditor/Private/Compiler/SceneStateBindingFunctionCompiler.h
--- a/Engine/Plugins/VirtualProduction/SceneState/Source/SceneStateBlueprintEditor/Private/Compiler/SceneStateBindingFunctionCompiler.h
+++ b/Engine/Plugins/VirtualProduction/SceneState/Source/SceneStateBlueprintEditor/Private/Compiler/SceneStateBindingFunctionCompiler.h
@@ -5,7 +5,7 @@
 #include "SceneStateRange.h"
 #include "UObject/ObjectPtr.h"
 
-struct FSceneStateTemplateData;
+class USceneStateTemplateData;
 
 namespace UE::SceneState::Editor
 {
@@ -14,7 +14,7 @@ namespace UE::SceneState::Editor
 class FBindingFunctionCompiler
 {
 public:
-	explicit FBindingFunctionCompiler(TNotNull<FSceneStateTemplateData*> InTemplateData)
+	explicit FBindingFunctionCompiler(TNotNull<USceneStateTemplateData*> InTemplateData)
 		: TemplateData(InTemplateData)
 	{
 	}
@@ -39,7 +39,7 @@ private:
 	void CompileTransitionFunctions();
 
 	/** The target to store the compiled functions */
-	FSceneStateTemplateData* TemplateData;
+	TObjectPtr<USceneStateTemplateData> TemplateData;
 
 	/** Map of the target struct id to its range of the functions in the generated class */
 	TMap<FGuid, FSceneStateRange> FunctionRangeMap;
diff --git a/Engine/Plugins/VirtualProduction/SceneState/Source/SceneStateBlueprintEditor/Private/Compiler/SceneStateBlueprintCompilerContext.cpp b/Engine/Plugins/VirtualProduction/SceneState/Source/SceneStateBlueprintEditor/Private/Compiler/SceneStateBlueprintCompilerContext.cpp
--- a/Engine/Plugins/VirtualProduction/SceneState/Source/SceneStateBlueprintEditor/Private/Compiler/SceneStateBlueprintCompilerContext.cpp
+++ b/Engine/Plugins/VirtualProduction/SceneState/Source/SceneStateBlueprintEditor/Private/Compiler/SceneStateBlueprintCompilerContext.cpp
@@ -65,6 +65,11 @@ FBlueprintCompilerContext::FBlueprintCompilerContext(USceneStateBlueprint* InBlu
 {
 }
 
+const UClass* FBlueprintCompilerContext::GetGeneratedClass() const
+{
+	return NewGeneratedClass;
+}
+
 void FBlueprintCompilerContext::SpawnNewClass(const FString& InNewClassName)
 {
 	check(Blueprint);
@@ -115,7 +120,8 @@ void FBlueprintCompilerContext::MergeUbergraphPagesIn(UEdGraph* InUbergraph)
 
 	USceneStateBlueprint* SceneStateBlueprint = CastChecked<USceneStateBlueprint>(Blueprint);
 
-	FSceneStateTemplateData& TemplateData = NewGeneratedClass->SharedTemplateData.Get<FSceneStateTemplateData>();
+	RecreateTemplateData();
+	check(NewGeneratedClass->TemplateData);
 
 	// Top Level State Machines
 	TArray<FSceneStateMachine> StateMachines;
@@ -127,8 +133,8 @@ void FBlueprintCompilerContext::MergeUbergraphPagesIn(UEdGraph* InUbergraph)
 	TMap<FObjectKey, uint16> StateMachineGraphToIndex;
 	StateMachineGraphToIndex.Reserve(SceneStateBlueprint->StateMachineGraphs.Num());
 
-	const int32 RootStateIndex = TemplateData.States.AddDefaulted();
-	const int32 RootStateMetadataIndex = TemplateData.StateMetadata.AddDefaulted();
+	const int32 RootStateIndex = NewGeneratedClass->TemplateData->States.AddDefaulted();
+	const int32 RootStateMetadataIndex = NewGeneratedClass->TemplateData->StateMetadata.AddDefaulted();
 
 	check(RootStateIndex == RootStateMetadataIndex);
 
@@ -161,22 +167,22 @@ void FBlueprintCompilerContext::MergeUbergraphPagesIn(UEdGraph* InUbergraph)
 		}
 	}
 
-	TemplateData.RootStateIndex = RootStateIndex;
+	NewGeneratedClass->TemplateData->RootStateIndex = RootStateIndex;
 
 	// Root State only considers the range of the Auto-Run State Machines
-	FSceneState& RootState = TemplateData.States[RootStateIndex];
-	RootState.StateMachineRange.Index = TemplateData.StateMachines.Num();
+	FSceneState& RootState = NewGeneratedClass->TemplateData->States[RootStateIndex];
+	RootState.StateMachineRange.Index = NewGeneratedClass->TemplateData->StateMachines.Num();
 	RootState.StateMachineRange.Count = AutoRunCount;
 
-	TemplateData.StateMachines.Append(MoveTemp(StateMachines));
+	NewGeneratedClass->TemplateData->StateMachines.Append(MoveTemp(StateMachines));
 
 	// Upgrade the map to absolute indices before baking it to the generated class' map
 	ToAbsoluteIndexMap(StateMachineIdToIndex, RootState.StateMachineRange.Index);
-	TemplateData.StateMachineIdToIndex.Append(MoveTemp(StateMachineIdToIndex));
+	NewGeneratedClass->TemplateData->StateMachineIdToIndex.Append(MoveTemp(StateMachineIdToIndex));
 
 	// Upgrade the map to absolute indices before baking it to the generated class' map
 	ToAbsoluteIndexMap(StateMachineGraphToIndex, RootState.StateMachineRange.Index);
-	TemplateData.StateMachineGraphToIndex.Append(MoveTemp(StateMachineGraphToIndex));
+	NewGeneratedClass->TemplateData->StateMachineGraphToIndex.Append(MoveTemp(StateMachineGraphToIndex));
 
 	UpdateTaskIndices();
 	CompileBindings();
@@ -188,16 +194,9 @@ void FBlueprintCompilerContext::OnPostCDOCompiled(const UObject::FPostCDOCompile
 	NewGeneratedClass->ResolveBindings();
 }
 
-FSceneStateTemplateData* FBlueprintCompilerContext::GetTemplateData() const
-{
-	return NewGeneratedClass
-		? NewGeneratedClass->SharedTemplateData.GetPtr<FSceneStateTemplateData>()
-		: nullptr;
-}
-
-UObject* FBlueprintCompilerContext::GetTargetOuter() const
+USceneStateTemplateData* FBlueprintCompilerContext::GetTemplateData()
 {
-	return NewGeneratedClass;
+	return NewGeneratedClass->TemplateData;
 }
 
 FTransitionGraphCompileResult FBlueprintCompilerContext::CompileTransitionGraph(USceneStateTransitionGraph* InTransitionGraph)
@@ -218,12 +217,23 @@ FTransitionGraphCompileResult FBlueprintCompilerContext::CompileTransitionGraph(
 	return CompileResult;
 }
 
+void FBlueprintCompilerContext::RecreateTemplateData()
+{
+	const TCHAR* TemplateDataName = TEXT("SceneStateTemplateData");
+
+	// Discard the current template data object as well as the unlikely event there's another object 
+	DiscardObject(NewGeneratedClass->TemplateData);
+	DiscardObject(NewGeneratedClass, TemplateDataName);
+
+	NewGeneratedClass->TemplateData = NewObject<USceneStateTemplateData>(NewGeneratedClass, TemplateDataName, RF_Public);
+}
+
 void FBlueprintCompilerContext::UpdateTaskIndices()
 {
-	FSceneStateTemplateData& TemplateData = NewGeneratedClass->SharedTemplateData.Get<FSceneStateTemplateData>();
+	check(NewGeneratedClass->TemplateData);
 
-	TConstArrayView<FSceneState> States = TemplateData.States;
-	FInstancedStructContainer& Tasks = TemplateData.Tasks;
+	TConstArrayView<FSceneState> States = NewGeneratedClass->TemplateData->States;
+	FInstancedStructContainer& Tasks = NewGeneratedClass->TemplateData->Tasks;
 
 	for (int32 StateIndex = 0; StateIndex < States.Num(); ++StateIndex)
 	{
@@ -244,9 +254,7 @@ void FBlueprintCompilerContext::UpdateTaskIndices()
 
 void FBlueprintCompilerContext::CompileBindings()
 {
-	FSceneStateTemplateData& TemplateData = NewGeneratedClass->SharedTemplateData.Get<FSceneStateTemplateData>();
-
-	FBindingCompiler BindingCompiler(*this, CastChecked<USceneStateBlueprint>(Blueprint), &TemplateData);
+	FBindingCompiler BindingCompiler(*this, CastChecked<USceneStateBlueprint>(Blueprint), NewGeneratedClass->TemplateData);
 	BindingCompiler.Compile();
 }
 
diff --git a/Engine/Plugins/VirtualProduction/SceneState/Source/SceneStateBlueprintEditor/Private/Compiler/SceneStateBlueprintCompilerContext.h b/Engine/Plugins/VirtualProduction/SceneState/Source/SceneStateBlueprintEditor/Private/Compiler/SceneStateBlueprintCompilerContext.h
--- a/Engine/Plugins/VirtualProduction/SceneState/Source/SceneStateBlueprintEditor/Private/Compiler/SceneStateBlueprintCompilerContext.h
+++ b/Engine/Plugins/VirtualProduction/SceneState/Source/SceneStateBlueprintEditor/Private/Compiler/SceneStateBlueprintCompilerContext.h
@@ -8,6 +8,7 @@
 
 class USceneStateBlueprint;
 class USceneStateGeneratedClass;
+class USceneStateTemplateData;
 
 namespace UE::SceneState::Editor
 {
@@ -17,6 +18,9 @@ class FBlueprintCompilerContext : public FKismetCompilerContext, public IStateMa
 public:
 	explicit FBlueprintCompilerContext(USceneStateBlueprint* InBlueprint, FCompilerResultsLog& InMessageLog, const FKismetCompilerOptions& InCompilerOptions);
 
+	/** Returns the new generated class for this compilation */
+	const UClass* GetGeneratedClass() const;
+
 private:
 	//~ Begin FKismetCompilerContext
 	virtual void SpawnNewClass(const FString& InNewClassName) override;
@@ -28,11 +32,13 @@ private:
 	//~ End FKismetCompilerContext
 
 	//~ Begin ISceneStateMachineCompilerContext
-	virtual FSceneStateTemplateData* GetTemplateData() const override;
-	virtual UObject* GetTargetOuter() const override;
+	virtual USceneStateTemplateData* GetTemplateData() override;
 	virtual FTransitionGraphCompileResult CompileTransitionGraph(USceneStateTransitionGraph* InTransitionGraph) override;
 	//~ End ISceneStateMachineCompilerContext
 
+	/** Recreates the template data owned by the generated class */
+	void RecreateTemplateData();
+
 	/** Sets the Task Index and Parent State Index of each Task */
 	void UpdateTaskIndices();
 
diff --git a/Engine/Plugins/VirtualProduction/SceneState/Source/SceneStateBlueprintEditor/Private/DebugExecutors/SceneStateDebugStateExecutor.cpp b/Engine/Plugins/VirtualProduction/SceneState/Source/SceneStateBlueprintEditor/Private/DebugExecutors/SceneStateDebugStateExecutor.cpp
--- a/Engine/Plugins/VirtualProduction/SceneState/Source/SceneStateBlueprintEditor/Private/DebugExecutors/SceneStateDebugStateExecutor.cpp
+++ b/Engine/Plugins/VirtualProduction/SceneState/Source/SceneStateBlueprintEditor/Private/DebugExecutors/SceneStateDebugStateExecutor.cpp
@@ -8,7 +8,11 @@ namespace UE::SceneState::Editor
 
 const FSceneState* FDebugStateExecutor::GetState(const FSceneStateExecutionContext& InExecutionContext) const
 {
-	return InExecutionContext.GetSharedTemplateData().FindStateFromNode(GetNodeKey());
+	if (const USceneStateTemplateData* TemplateData = InExecutionContext.GetTemplateData())
+	{
+		return TemplateData->FindStateFromNode(GetNodeKey());
+	}
+	return nullptr;
 }
 
 void FDebugStateExecutor::OnStart(const FSceneStateExecutionContext& InExecutionContext)
diff --git a/Engine/Plugins/VirtualProduction/SceneState/Source/SceneStateBlueprintEditor/Private/DebugExecutors/SceneStateDebugTaskExecutor.cpp b/Engine/Plugins/VirtualProduction/SceneState/Source/SceneStateBlueprintEditor/Private/DebugExecutors/SceneStateDebugTaskExecutor.cpp
--- a/Engine/Plugins/VirtualProduction/SceneState/Source/SceneStateBlueprintEditor/Private/DebugExecutors/SceneStateDebugTaskExecutor.cpp
+++ b/Engine/Plugins/VirtualProduction/SceneState/Source/SceneStateBlueprintEditor/Private/DebugExecutors/SceneStateDebugTaskExecutor.cpp
@@ -12,16 +12,20 @@ namespace UE::SceneState::Editor
 
 const FSceneStateTask* FDebugTaskExecutor::GetTask(const FSceneStateExecutionContext& InExecutionContext) const
 {
-	return InExecutionContext.GetSharedTemplateData().FindTaskFromNode(GetNodeKey());
+	if (const USceneStateTemplateData* TemplateData = InExecutionContext.GetTemplateData())
+	{
+		return TemplateData->FindTaskFromNode(GetNodeKey());
+	}
+	return nullptr;
 }
 
-const FSceneState* FDebugTaskExecutor::GetState(const FSceneStateTemplateData& InTemplateData) const
+const FSceneState* FDebugTaskExecutor::GetState(TNotNull<const USceneStateTemplateData*> InTemplateData) const
 {
 	if (const USceneStateMachineTaskNode* const TaskNode = Cast<USceneStateMachineTaskNode>(GetNodeKey().ResolveObjectPtr()))
 	{
 		if (const USceneStateMachineStateNode* const StateNode = USceneStateMachineGraphSchema::FindConnectedStateNode(TaskNode))
 		{
-			return InTemplateData.FindStateFromNode(StateNode);
+			return InTemplateData->FindStateFromNode(StateNode);
 		}
 	}
 	return nullptr;
@@ -29,7 +33,11 @@ const FSceneState* FDebugTaskExecutor::GetState(const FSceneStateTemplateData& I
 
 void FDebugTaskExecutor::Setup(const FSceneStateExecutionContext& InExecutionContext, const FSceneStateTask& InTask)
 {
-	const FSceneStateTemplateData& TemplateData = InExecutionContext.GetSharedTemplateData();
+	const USceneStateTemplateData* TemplateData = InExecutionContext.GetTemplateData();
+	if (!TemplateData)
+	{
+		return;
+	}
 
 	const FSceneState* ParentState = GetState(TemplateData);
 	if (!ParentState)
@@ -46,7 +54,7 @@ void FDebugTaskExecutor::Setup(const FSceneStateExecutionContext& InExecutionCon
 	// and only set the task instance template for the task this executor cares about
 	TArray<FConstStructView> TemplateTaskInstances;
 	TemplateTaskInstances.SetNum(TaskRange.Count);
-	TemplateTaskInstances[TaskRelativeIndex] = TemplateData.TaskInstances[TaskAbsoluteIndex];
+	TemplateTaskInstances[TaskRelativeIndex] = TemplateData->TaskInstances[TaskAbsoluteIndex];
 
 	ParentState->AllocateTaskInstances(InExecutionContext, TemplateTaskInstances);
 }
diff --git a/Engine/Plugins/VirtualProduction/SceneState/Source/SceneStateBlueprintEditor/Private/DebugExecutors/SceneStateDebugTaskExecutor.h b/Engine/Plugins/VirtualProduction/SceneState/Source/SceneStateBlueprintEditor/Private/DebugExecutors/SceneStateDebugTaskExecutor.h
--- a/Engine/Plugins/VirtualProduction/SceneState/Source/SceneStateBlueprintEditor/Private/DebugExecutors/SceneStateDebugTaskExecutor.h
+++ b/Engine/Plugins/VirtualProduction/SceneState/Source/SceneStateBlueprintEditor/Private/DebugExecutors/SceneStateDebugTaskExecutor.h
@@ -21,7 +21,7 @@ protected:
 	const FSceneStateTask* GetTask(const FSceneStateExecutionContext& InExecutionContext) const;
 
 	/** Retrieves the state corresponding to the stored state node */
-	const FSceneState* GetState(const FSceneStateTemplateData& InTemplateData) const;
+	const FSceneState* GetState(TNotNull<const USceneStateTemplateData*> InTemplateData) const;
 
 	/** Called to allocate the task instance */
 	void Setup(const FSceneStateExecutionContext& InExecutionContext, const FSceneStateTask& InTask);
diff --git a/Engine/Plugins/VirtualProduction/SceneState/Source/SceneStateMachineEditor/Private/SceneStateMachineCompiler.cpp b/Engine/Plugins/VirtualProduction/SceneState/Source/SceneStateMachineEditor/Private/SceneStateMachineCompiler.cpp
--- a/Engine/Plugins/VirtualProduction/SceneState/Source/SceneStateMachineEditor/Private/SceneStateMachineCompiler.cpp
+++ b/Engine/Plugins/VirtualProduction/SceneState/Source/SceneStateMachineEditor/Private/SceneStateMachineCompiler.cpp
@@ -57,13 +57,12 @@ TMap<uint16, FInstancedPropertyBag> ConvertParameterMap(FSceneStateRange InTrans
 } // Private
 
 FStateMachineCompiler::FStateMachineCompiler(USceneStateMachineGraph* InGraph, IStateMachineCompilerContext& InCompilerContext)
-	: Outer(InCompilerContext.GetTargetOuter())
-	, TemplateData(InCompilerContext.GetTemplateData())
+	: TemplateData(InCompilerContext.GetTemplateData())
 	, StateMachineGraph(InGraph)
 	, Blueprint(FBlueprintEditorUtils::FindBlueprintForGraph(InGraph))
 	, Context(InCompilerContext)
 {
-	check(TemplateData && StateMachineGraph && Outer);
+	check(TemplateData && StateMachineGraph);
 }
 
 FSceneStateMachine FStateMachineCompiler::Compile()
@@ -262,7 +261,7 @@ void FStateMachineCompiler::CompileTasks(FSceneState& InNewState, const USceneSt
 {
 	FStateMachineTaskCompiler::FCompileResult TaskCompileResult;
 	{
-		FStateMachineTaskCompiler TaskCompiler(InStateNode->GetTaskPin(), /*Outer*/Outer);
+		FStateMachineTaskCompiler TaskCompiler(InStateNode->GetTaskPin(), /*Outer*/TemplateData);
 		TaskCompiler.Compile(TaskCompileResult);
 	}
 
@@ -282,7 +281,7 @@ void FStateMachineCompiler::CompileTasks(FSceneState& InNewState, const USceneSt
 	TemplateData->TaskMetadata.Append(MoveTemp(TaskCompileResult.TaskMetadata));
 	TemplateData->TaskPrerequisites.Append(MoveTemp(TaskCompileResult.TaskPrerequisites));
 
-	InNewState.InstanceTaskObjects(Outer
+	InNewState.InstanceTaskObjects(TemplateData
 		, GetStructViews(TemplateData->TaskInstances, InNewState.TaskRange)
 		, TaskCompileResult.TaskInstances
 		, [this](FObjectDuplicationParameters& InParams)->UObject*
diff --git a/Engine/Plugins/VirtualProduction/SceneState/Source/SceneStateMachineEditor/Private/Widgets/SSceneStateMachineStateNode.cpp b/Engine/Plugins/VirtualProduction/SceneState/Source/SceneStateMachineEditor/Private/Widgets/SSceneStateMachineStateNode.cpp
--- a/Engine/Plugins/VirtualProduction/SceneState/Source/SceneStateMachineEditor/Private/Widgets/SSceneStateMachineStateNode.cpp
+++ b/Engine/Plugins/VirtualProduction/SceneState/Source/SceneStateMachineEditor/Private/Widgets/SSceneStateMachineStateNode.cpp
@@ -190,7 +190,7 @@ void SStateMachineStateNode::GetNodeInfoPopups(FNodeInfoContext* InContext, TArr
 		return;
 	}
 
-	const FSceneStateTemplateData* TemplateData = RootState->GetSharedTemplateData().GetPtr<const FSceneStateTemplateData>();
+	const USceneStateTemplateData* TemplateData = RootState->GetTemplateData();
 	if (!TemplateData)
 	{
 		return;
diff --git a/Engine/Plugins/VirtualProduction/SceneState/Source/SceneStateMachineEditor/Private/Widgets/SSceneStateMachineTaskNode.cpp b/Engine/Plugins/VirtualProduction/SceneState/Source/SceneStateMachineEditor/Private/Widgets/SSceneStateMachineTaskNode.cpp
--- a/Engine/Plugins/VirtualProduction/SceneState/Source/SceneStateMachineEditor/Private/Widgets/SSceneStateMachineTaskNode.cpp
+++ b/Engine/Plugins/VirtualProduction/SceneState/Source/SceneStateMachineEditor/Private/Widgets/SSceneStateMachineTaskNode.cpp
@@ -190,7 +190,7 @@ void SStateMachineTaskNode::GetNodeInfoPopups(FNodeInfoContext* InContext, TArra
 		return;
 	}
 
-	const FSceneStateTemplateData* TemplateData = RootState->GetSharedTemplateData().GetPtr<const FSceneStateTemplateData>();
+	const USceneStateTemplateData* TemplateData = RootState->GetTemplateData();
 	if (!TemplateData)
 	{
 		return;
diff --git a/Engine/Plugins/VirtualProduction/SceneState/Source/SceneStateMachineEditor/Public/ISceneStateMachineCompilerContext.h b/Engine/Plugins/VirtualProduction/SceneState/Source/SceneStateMachineEditor/Public/ISceneStateMachineCompilerContext.h
--- a/Engine/Plugins/VirtualProduction/SceneState/Source/SceneStateMachineEditor/Public/ISceneStateMachineCompilerContext.h
+++ b/Engine/Plugins/VirtualProduction/SceneState/Source/SceneStateMachineEditor/Public/ISceneStateMachineCompilerContext.h
@@ -5,8 +5,8 @@
 #include "HAL/Platform.h"
 #include "UObject/NameTypes.h"
 
+class USceneStateTemplateData;
 class USceneStateTransitionGraph;
-struct FSceneStateTemplateData;
 
 namespace UE::SceneState::Editor
 {
@@ -41,11 +41,9 @@ class IStateMachineCompilerContext
 {
 public:
 	/** Gets the template data to compile the objects to */
-	virtual FSceneStateTemplateData* GetTemplateData() const = 0;
-
-	/** Returns the object to outer all compiled instanced UObjects to */
-	virtual UObject* GetTargetOuter() const = 0;
+	virtual USceneStateTemplateData* GetTemplateData() = 0;
 
+	/** Requests compilation of the given transition graph */
 	virtual FTransitionGraphCompileResult CompileTransitionGraph(USceneStateTransitionGraph* InTransitionGraph) = 0;
 };
 
diff --git a/Engine/Plugins/VirtualProduction/SceneState/Source/SceneStateMachineEditor/Public/SceneStateMachineCompiler.h b/Engine/Plugins/VirtualProduction/SceneState/Source/SceneStateMachineEditor/Public/SceneStateMachineCompiler.h
--- a/Engine/Plugins/VirtualProduction/SceneState/Source/SceneStateMachineEditor/Public/SceneStateMachineCompiler.h
+++ b/Engine/Plugins/VirtualProduction/SceneState/Source/SceneStateMachineEditor/Public/SceneStateMachineCompiler.h
@@ -17,8 +17,8 @@ class USceneStateMachineConduitNode;
 class USceneStateMachineGraph;
 class USceneStateMachineStateNode;
 class USceneStateMachineTransitionNode;
+class USceneStateTemplateData;
 struct FSceneStateMachine;
-struct FSceneStateTemplateData;
 
 namespace UE::SceneState::Editor
 {
@@ -66,11 +66,8 @@ private:
 
 	UObject* DuplicateObject(FObjectDuplicationParameters& InDuplicationParams);
 
-	/** Outer of the all the objects to create here */
-	TObjectPtr<UObject> Outer;
-
 	/** Template data where all the compiled objects will be allocated to */
-	FSceneStateTemplateData* TemplateData;
+	TObjectPtr<USceneStateTemplateData> TemplateData;
 
 	/** Reference Graph that will build out the runtime State Machine object */
 	TObjectPtr<USceneStateMachineGraph> StateMachineGraph;
