Default initialize raw ptr to null + fixed outdated comment

diff --git a/Engine/Plugins/VirtualProduction/SceneState/Source/SceneState/Public/SceneStateExecutionContext.h b/Engine/Plugins/VirtualProduction/SceneState/Source/SceneState/Public/SceneStateExecutionContext.h
--- a/Engine/Plugins/VirtualProduction/SceneState/Source/SceneState/Public/SceneStateExecutionContext.h
+++ b/Engine/Plugins/VirtualProduction/SceneState/Source/SceneState/Public/SceneStateExecutionContext.h
@@ -104,8 +104,8 @@ struct FSceneStateExecutionContext
 
 	struct FSetupParams
 	{
-		/** Shared memory of the template data holding all the state machines, states, etc. */
-		const USceneStateTemplateData* TemplateData;
+		/** Template data object holding all the state machines, states, etc. */
+		const USceneStateTemplateData* TemplateData = nullptr;
 		/** Root object instance */
 		UObject* RootObject = nullptr;
 		/** Context object for utility purposes (e.g. to get world from context, etc) */
