Motion Design Scene State: changed error log to warning.

diff --git a/Engine/Plugins/VirtualProduction/SceneState/Source/SceneState/Private/SceneStateObject.cpp b/Engine/Plugins/VirtualProduction/SceneState/Source/SceneState/Private/SceneStateObject.cpp
--- a/Engine/Plugins/VirtualProduction/SceneState/Source/SceneState/Private/SceneStateObject.cpp
+++ b/Engine/Plugins/VirtualProduction/SceneState/Source/SceneState/Private/SceneStateObject.cpp
@@ -56,7 +56,7 @@ void USceneStateObject::Setup()
 
 	if (!TemplateData)
 	{
-		UE_LOG(LogSceneState, Error, TEXT("[%s] did not find a valid template data object."), *GetFullName());
+		UE_LOG(LogSceneState, Warning, TEXT("[%s] did not find a valid template data object. Scene State will not execute in this instance."), *GetFullName());
 		return;
 	}
 
