Motion Design Scene State: added actor reference for the spawned actor to be useable outside its task

diff --git a/Engine/Plugins/VirtualProduction/SceneState/Source/SceneStateBinding/Public/SceneStatePropertyReference.h b/Engine/Plugins/VirtualProduction/SceneState/Source/SceneStateBinding/Public/SceneStatePropertyReference.h
--- a/Engine/Plugins/VirtualProduction/SceneState/Source/SceneStateBinding/Public/SceneStatePropertyReference.h
+++ b/Engine/Plugins/VirtualProduction/SceneState/Source/SceneStateBinding/Public/SceneStatePropertyReference.h
@@ -35,6 +35,12 @@ struct FSceneStatePropertyReference
 	/** Max number of indices that can be used */
 	static constexpr uint16 IndexCapacity = std::numeric_limits<uint16>::max();
 
+	/** Returns whether the index is a valid index. Does not guarantee that the reference itself is to a valid property */
+	bool IsValidIndex() const
+	{
+		return ReferenceIndex != IndexCapacity;
+	}
+
 	UPROPERTY()
 	uint16 ReferenceIndex = IndexCapacity;
 };
diff --git a/Engine/Plugins/VirtualProduction/SceneState/Source/SceneStateTasks/Private/SceneStateSpawnActorTask.cpp b/Engine/Plugins/VirtualProduction/SceneState/Source/SceneStateTasks/Private/SceneStateSpawnActorTask.cpp
--- a/Engine/Plugins/VirtualProduction/SceneState/Source/SceneStateTasks/Private/SceneStateSpawnActorTask.cpp
+++ b/Engine/Plugins/VirtualProduction/SceneState/Source/SceneStateTasks/Private/SceneStateSpawnActorTask.cpp
@@ -5,6 +5,7 @@
 #include "GameFramework/Actor.h"
 #include "SceneStateExecutionContext.h"
 #include "SceneStateTasksLog.h"
+#include "Setters/SceneStateSetterUtils.h"
 
 #if WITH_EDITOR
 #include "ActorFactories/ActorFactory.h"
@@ -32,13 +33,16 @@ void FSceneStateSpawnActorTask::OnPostEditChange(UE::SceneState::FTaskEditChange
 
 void FSceneStateSpawnActorTask::OnStart(const FSceneStateExecutionContext& InContext, FStructView InTaskInstance) const
 {
+	FInstanceDataType& Instance = InTaskInstance.Get<FInstanceDataType>();
+
+	AActor* SpawnedActor = nullptr;
+
 	ON_SCOPE_EXIT
 	{
+		SetSpawnedActor(InContext, SpawnedActor, Instance.SpawnedActor);
 		Finish(InContext, InTaskInstance);
 	};
 
-	FInstanceDataType& Instance = InTaskInstance.Get<FInstanceDataType>();
-
 	AActor* const TemplateActor = Instance.ActorTemplate.Template;
 	if (!TemplateActor)
 	{
@@ -75,9 +79,10 @@ void FSceneStateSpawnActorTask::OnStart(const FSceneStateExecutionContext& InCon
 	SpawnParameters.SpawnCollisionHandlingOverride = Instance.SpawnCollisionHandling;
 	SpawnParameters.ObjectFlags = RF_Transient | RF_Transactional;
 
-	if (AActor* Actor = World->SpawnActor(TemplateActor->GetClass(), &Instance.SpawnTransform, SpawnParameters))
+	SpawnedActor = World->SpawnActor(TemplateActor->GetClass(), &Instance.SpawnTransform, SpawnParameters);
+	if (SpawnedActor)
 	{
-		OnActorSpawned(Actor, InTaskInstance);
+		OnActorSpawned(SpawnedActor, InTaskInstance);
 	}
 }
 
@@ -155,3 +160,31 @@ void FSceneStateSpawnActorTask::UpdateActorTemplate(UObject* InOuter, FStructVie
 	World->DestroyActor(SpawnedActor, bNetForce, bCallModify);
 }
 #endif
+
+void FSceneStateSpawnActorTask::SetSpawnedActor(const FSceneStateExecutionContext& InContext, AActor* InSpawnedActor, const FSceneStatePropertyReference& InSpawnedActorReference) const
+{
+	// Spawned actor reference is optional so it's ok if it is unset.
+	if (!InSpawnedActorReference.IsValidIndex())
+	{
+		return;
+	}
+
+	UE::SceneState::FResolvePropertyResult Result;
+	if (ResolveProperty(InContext, InSpawnedActorReference, Result))
+	{
+		if (!UE::SceneState::SetValue<AActor*>(Result.ValuePtr, *Result.ResolvedReference, InSpawnedActor))
+		{
+			UE_LOG(LogSceneStateTasks, Warning, TEXT("[%s] Spawned Actor Reference was not set for spawned actor '%s' (Class: %s). Failed to set value.")
+				, *InContext.GetExecutionContextName()
+				, *InSpawnedActor->GetActorNameOrLabel()
+				, *GetNameSafe(InSpawnedActor->GetClass()));
+		}
+	}
+	else
+	{
+		UE_LOG(LogSceneStateTasks, Warning, TEXT("[%s] Spawned Actor Reference was not set for spawned actor '%s' (Class: %s). Failed to resolve property.")
+			, *InContext.GetExecutionContextName()
+			, *InSpawnedActor->GetActorNameOrLabel()
+			, *GetNameSafe(InSpawnedActor->GetClass()));
+	}
+}
diff --git a/Engine/Plugins/VirtualProduction/SceneState/Source/SceneStateTasks/Public/SceneStateSpawnActorTask.h b/Engine/Plugins/VirtualProduction/SceneState/Source/SceneStateTasks/Public/SceneStateSpawnActorTask.h
--- a/Engine/Plugins/VirtualProduction/SceneState/Source/SceneStateTasks/Public/SceneStateSpawnActorTask.h
+++ b/Engine/Plugins/VirtualProduction/SceneState/Source/SceneStateTasks/Public/SceneStateSpawnActorTask.h
@@ -3,6 +3,7 @@
 #pragma once
 
 #include "Engine/EngineTypes.h"
+#include "SceneStatePropertyReference.h"
 #include "Tasks/SceneStateTask.h"
 #include "Tasks/SceneStateTaskInstance.h"
 #include "Templates/SubclassOf.h"
@@ -37,8 +38,12 @@ struct FSceneStateSpawnActorTaskInstance : public FSceneStateTaskInstance
 	FTransform SpawnTransform = FTransform::Identity;
 
 	/** Method for resolving collisions at the spawn point. Default (unspecified) means no override, use the actor's setting */
-	UPROPERTY(EditAnywhere, AdvancedDisplay, Category="Scene State")
+	UPROPERTY(EditAnywhere, Category="Scene State")
 	ESpawnActorCollisionHandlingMethod SpawnCollisionHandling = ESpawnActorCollisionHandlingMethod::Undefined;
+
+	/** Optional: Sets the actor that was spawned from this task */
+	UPROPERTY(EditAnywhere, Category="Scene State", meta=(RefType="/Script/Engine.Actor"))
+	FSceneStatePropertyReference SpawnedActor;
 };
 
 /** Spawns an actor of a given class */
@@ -63,6 +68,14 @@ protected:
 	void UpdateActorTemplate(UObject* InOuter, FStructView InTaskInstance) const;
 #endif
 
+	/**
+	 * Updates the spawned actor reference value (if set) to the newly spawned actor
+	 * @param InContext the execution context that is running this task
+	 * @param InSpawnedActor the value to set the property reference to
+	 * @param InSpawnedActorReference the actor property reference to set
+	 */
+	void SetSpawnedActor(const FSceneStateExecutionContext& InContext, AActor* InSpawnedActor, const FSceneStatePropertyReference& InSpawnedActorReference) const;
+
 	/**
 	 * Determines whether the actor should be spawned
 	 * @param InTaskInstance data view to the task instance of the execution.
