Motion Desgin: Fix crash when specifying a blueprint class in the 'Actor class' property in 'Spawn Actor' tasks.

diff --git a/Engine/Plugins/VirtualProduction/SceneState/Source/SceneStateTasks/Private/SceneStateSpawnActorTask.cpp b/Engine/Plugins/VirtualProduction/SceneState/Source/SceneStateTasks/Private/SceneStateSpawnActorTask.cpp
--- a/Engine/Plugins/VirtualProduction/SceneState/Source/SceneStateTasks/Private/SceneStateSpawnActorTask.cpp
+++ b/Engine/Plugins/VirtualProduction/SceneState/Source/SceneStateTasks/Private/SceneStateSpawnActorTask.cpp
@@ -112,7 +112,11 @@ void FSceneStateSpawnActorTask::UpdateActorTemplate(UObject* InOuter, FStructVie
 		return;
 	}
 
-	UActorFactory* ActorFactory = FActorFactoryAssetProxy::GetFactoryForAssetObject(ActorClass);
+	const FAssetData ActorClassAssetData(ActorClass);
+
+	// FAssetData::GetAsset loads the asset synchronously if not already loaded. 
+	// This editor-only function is called only in PostEditChange when the ActorClass has changed, so no need to load this async for now.
+	UActorFactory* ActorFactory = FActorFactoryAssetProxy::GetFactoryForAssetObject(ActorClassAssetData.GetAsset());
 	if (!ActorFactory)
 	{
 		TemplateActor = nullptr;
@@ -121,7 +125,7 @@ void FSceneStateSpawnActorTask::UpdateActorTemplate(UObject* InOuter, FStructVie
 	}
 
 	FText ErrorText;
-	if (!ActorFactory->CanCreateActorFrom(FAssetData(ActorClass), ErrorText))
+	if (!ActorFactory->CanCreateActorFrom(ActorClassAssetData, ErrorText))
 	{
 		TemplateActor = nullptr;
 		UE_LOG(LogSceneStateTasks, Error, TEXT("Unable to create actor of class %s. Reason: %s"), *ActorClass->GetName(), *ErrorText.ToString());
@@ -131,7 +135,7 @@ void FSceneStateSpawnActorTask::UpdateActorTemplate(UObject* InOuter, FStructVie
 	FActorSpawnParameters SpawnParams;
 	SpawnParams.ObjectFlags = RF_Transient | RF_Transactional;
 
-	AActor* const SpawnedActor = ActorFactory->CreateActor(ActorClass, World->PersistentLevel, FTransform::Identity, SpawnParams);
+	AActor* const SpawnedActor = ActorFactory->CreateActor(ActorClassAssetData.GetAsset(), World->PersistentLevel, FTransform::Identity, SpawnParams);
 	if (!SpawnedActor)
 	{
 		TemplateActor = nullptr;
