Motion Design: fix crash on shutdown when scene state editor containing bindings is left opened.

diff --git a/Engine/Plugins/VirtualProduction/SceneState/Source/SceneState/Private/SceneStateGeneratedClass.cpp b/Engine/Plugins/VirtualProduction/SceneState/Source/SceneState/Private/SceneStateGeneratedClass.cpp
--- a/Engine/Plugins/VirtualProduction/SceneState/Source/SceneState/Private/SceneStateGeneratedClass.cpp
+++ b/Engine/Plugins/VirtualProduction/SceneState/Source/SceneState/Private/SceneStateGeneratedClass.cpp
@@ -47,10 +47,6 @@ void USceneStateGeneratedClass::PostLoad()
 void USceneStateGeneratedClass::BeginDestroy()
 {
 	Super::BeginDestroy();
-	if (TemplateData)
-	{
-		TemplateData->Reset();
-	}
 }
 
 void USceneStateGeneratedClass::ResolveBindings()
