Motion Design Scene State: fix issue with debug controls showing properties that are not instance editable. The debug controls tab should look very similar as the level editor details panel only showing instance editable properties.
+ also made the details panel for the scene state editor selection disabled to avoid user confusion of what to edit when debugging.

#Synced-CL 48854462

diff --git a/Engine/Plugins/VirtualProduction/SceneState/Source/SceneStateBlueprintEditor/Private/DebugControls/SSceneStateDebugControls.cpp b/Engine/Plugins/VirtualProduction/SceneState/Source/SceneStateBlueprintEditor/Private/DebugControls/SSceneStateDebugControls.cpp
--- a/Engine/Plugins/VirtualProduction/SceneState/Source/SceneStateBlueprintEditor/Private/DebugControls/SSceneStateDebugControls.cpp
+++ b/Engine/Plugins/VirtualProduction/SceneState/Source/SceneStateBlueprintEditor/Private/DebugControls/SSceneStateDebugControls.cpp
@@ -34,6 +34,8 @@ void SDebugControls::Construct(const FArguments& InArgs, const TSharedRef<FScene
 	DebugControlsDetailsView = CreateDebugControlsDetailsView(InBlueprintEditor);
 	DebugControlsDetailsView->SetObject(DebugControlsTool->GetDebugControlsObject());
 
+	DebuggedObjectDetailsView = CreateDebuggedObjectDetailsView(InBlueprintEditor);
+
 	ContentWidget = CreateContentWidget();
 	PlaceholderWidget = CreatePlaceholderWidget();
 
@@ -58,8 +60,8 @@ SDebugControls::~SDebugControls()
 
 void SDebugControls::Refresh()
 {
-	DebugControlsTool->UpdateDebuggedObject();
-	DebugControlsDetailsView->ForceRefresh();
+	USceneStateObject* const DebuggedObject = DebugControlsTool->GetDebuggedObject();
+	DebuggedObjectDetailsView->SetObject(DebuggedObject);
 }
 
 void SDebugControls::Tick(const FGeometry& InAllottedGeometry, const double InCurrentTime, const float InDeltaTime)
@@ -86,9 +88,15 @@ TSharedRef<SWidget> SDebugControls::CreateContentWidget()
 			CreateDebugControlsToolbar(DebugControlsTool->GetCommandList())
 		]
 		+ SVerticalBox::Slot()
-		.FillHeight(1.f)
+		.AutoHeight()
 		[
 			DebugControlsDetailsView.ToSharedRef()	
+		]
+		+ SVerticalBox::Slot()
+		.FillHeight(1.f)
+		.Padding(0, 5, 0, 0)
+		[
+			DebuggedObjectDetailsView.ToSharedRef()	
 		];
 }
 
@@ -108,10 +116,13 @@ TSharedRef<SWidget> SDebugControls::CreatePlaceholderWidget()
 TSharedRef<IDetailsView> SDebugControls::CreateDebugControlsDetailsView(const TSharedRef<FSceneStateBlueprintEditor>& InBlueprintEditor) const
 {
 	FDetailsViewArgs DetailsViewArgs;
-	DetailsViewArgs.DefaultsOnlyVisibility = EEditDefaultsOnlyNodeVisibility::Automatic;
+	DetailsViewArgs.DefaultsOnlyVisibility = EEditDefaultsOnlyNodeVisibility::Hide;
 	DetailsViewArgs.HostCommandList = InBlueprintEditor->GetToolkitCommands();
 	DetailsViewArgs.HostTabManager = InBlueprintEditor->GetTabManager();
 	DetailsViewArgs.NameAreaSettings = FDetailsViewArgs::HideNameArea;
+	DetailsViewArgs.bAllowSearch = false;
+	DetailsViewArgs.bShowOptions = false;
+	DetailsViewArgs.bShowSectionSelector = false;
 
 	FPropertyEditorModule& PropertyEditorModule = FModuleManager::Get().LoadModuleChecked<FPropertyEditorModule>("PropertyEditor");
 
@@ -123,6 +134,21 @@ TSharedRef<IDetailsView> SDebugControls::CreateDebugControlsDetailsView(const TS
 	return DetailsView;
 }
 
+TSharedRef<IDetailsView> SDebugControls::CreateDebuggedObjectDetailsView(const TSharedRef<FSceneStateBlueprintEditor>& InBlueprintEditor) const
+{
+	FDetailsViewArgs DetailsViewArgs;
+	DetailsViewArgs.DefaultsOnlyVisibility = EEditDefaultsOnlyNodeVisibility::Hide;
+	DetailsViewArgs.HostCommandList = InBlueprintEditor->GetToolkitCommands();
+	DetailsViewArgs.HostTabManager = InBlueprintEditor->GetTabManager();
+	DetailsViewArgs.NameAreaSettings = FDetailsViewArgs::HideNameArea;
+	DetailsViewArgs.bAllowSearch = true;
+	DetailsViewArgs.bShowOptions = true;
+	DetailsViewArgs.bShowSectionSelector = true;
+	
+	FPropertyEditorModule& PropertyEditorModule = FModuleManager::Get().LoadModuleChecked<FPropertyEditorModule>("PropertyEditor");
+	return PropertyEditorModule.CreateDetailView(DetailsViewArgs);
+}
+
 void SDebugControls::OnBlueprintDebugObjectChanged(const Graph::FBlueprintDebugObjectChange& InChange)
 {
 	TSharedPtr<FSceneStateBlueprintEditor> BlueprintEditor = BlueprintEditorWeak.Pin();
diff --git a/Engine/Plugins/VirtualProduction/SceneState/Source/SceneStateBlueprintEditor/Private/DebugControls/SSceneStateDebugControls.h b/Engine/Plugins/VirtualProduction/SceneState/Source/SceneStateBlueprintEditor/Private/DebugControls/SSceneStateDebugControls.h
--- a/Engine/Plugins/VirtualProduction/SceneState/Source/SceneStateBlueprintEditor/Private/DebugControls/SSceneStateDebugControls.h
+++ b/Engine/Plugins/VirtualProduction/SceneState/Source/SceneStateBlueprintEditor/Private/DebugControls/SSceneStateDebugControls.h
@@ -56,12 +56,18 @@ private:
 	/** Creates a new details view for debug controls */
 	TSharedRef<IDetailsView> CreateDebugControlsDetailsView(const TSharedRef<FSceneStateBlueprintEditor>& InBlueprintEditor) const;
 
+	/** Creates a new details view for debugged object */
+	TSharedRef<IDetailsView> CreateDebuggedObjectDetailsView(const TSharedRef<FSceneStateBlueprintEditor>& InBlueprintEditor) const;
+
 	/** Called when the blueprint debug object has changed */
 	void OnBlueprintDebugObjectChanged(const Graph::FBlueprintDebugObjectChange& InChange);
 
 	/** Details view of the debug controls object */
 	TSharedPtr<IDetailsView> DebugControlsDetailsView;
 
+	/** Details view of the debugged object */
+	TSharedPtr<IDetailsView> DebuggedObjectDetailsView;
+
 	/** The blueprint editor owning this widget */
 	TWeakPtr<FSceneStateBlueprintEditor> BlueprintEditorWeak;
 
diff --git a/Engine/Plugins/VirtualProduction/SceneState/Source/SceneStateBlueprintEditor/Private/DebugControls/SceneStateDebugControlsObject.h b/Engine/Plugins/VirtualProduction/SceneState/Source/SceneStateBlueprintEditor/Private/DebugControls/SceneStateDebugControlsObject.h
--- a/Engine/Plugins/VirtualProduction/SceneState/Source/SceneStateBlueprintEditor/Private/DebugControls/SceneStateDebugControlsObject.h
+++ b/Engine/Plugins/VirtualProduction/SceneState/Source/SceneStateBlueprintEditor/Private/DebugControls/SceneStateDebugControlsObject.h
@@ -15,10 +15,6 @@ class USceneStateDebugControlsObject : public UObject
 	GENERATED_BODY()
 
 public:
-	/** Currently debugged scene state object */
-	UPROPERTY()
-	TWeakObjectPtr<USceneStateObject> DebuggedObjectWeak;
-
 	/** Debug events to push */
 	UPROPERTY(EditAnywhere, Category="Event")
 	TArray<FSceneStateEventTemplate> Events;
diff --git a/Engine/Plugins/VirtualProduction/SceneState/Source/SceneStateBlueprintEditor/Private/DebugControls/SceneStateDebugControlsObjectDetails.cpp b/Engine/Plugins/VirtualProduction/SceneState/Source/SceneStateBlueprintEditor/Private/DebugControls/SceneStateDebugControlsObjectDetails.cpp
--- a/Engine/Plugins/VirtualProduction/SceneState/Source/SceneStateBlueprintEditor/Private/DebugControls/SceneStateDebugControlsObjectDetails.cpp
+++ b/Engine/Plugins/VirtualProduction/SceneState/Source/SceneStateBlueprintEditor/Private/DebugControls/SceneStateDebugControlsObjectDetails.cpp
@@ -8,7 +8,6 @@
 #include "PropertyCustomizationHelpers.h"
 #include "SceneStateDebugControlsObject.h"
 #include "SceneStateEventUtils.h"
-#include "SceneStateObject.h"
 #include "Widgets/Input/SButton.h"
 #include "Widgets/SBoxPanel.h"
 #include "Widgets/Text/STextBlock.h"
@@ -22,61 +21,48 @@ void FDebugControlsObjectDetails::CustomizeDetails(IDetailLayoutBuilder& InDetai
 {
 	DebugControls = InDetailBuilder.GetObjectsOfTypeBeingCustomized<USceneStateDebugControlsObject>();
 	CustomizeEventDetails(InDetailBuilder);
-	CustomizeDebuggedObjectDetails(InDetailBuilder);
 }
 
 void FDebugControlsObjectDetails::CustomizeEventDetails(IDetailLayoutBuilder& InDetailBuilder)
 {
-	EventsHandle = InDetailBuilder.GetProperty(GET_MEMBER_NAME_CHECKED(USceneStateDebugControlsObject, Events));
+	TSharedRef<IPropertyHandle> EventsHandle = InDetailBuilder.GetProperty(GET_MEMBER_NAME_CHECKED(USceneStateDebugControlsObject, Events));
 	EventsHandle->MarkHiddenByCustomization();
 
 	IDetailCategoryBuilder& EventsCategory = InDetailBuilder.EditCategory(TEXT("Events"));
-	EventsCategory.AddProperty(EventsHandle);
-}
-
-USceneStateObject* FDebugControlsObjectDetails::GetDebuggedObject(int32 InIndex) const
-{
-	if (!DebugControls.IsValidIndex(InIndex))
-	{
-		return nullptr;
-	}
 
-	USceneStateDebugControlsObject* const DebugControlsObject = DebugControls[InIndex].Get();
-	if (!DebugControlsObject)
-	{
-		return nullptr;
-	}
+	TSharedRef<SWidget> HeaderContentWidget = SNew(SHorizontalBox)
+		+ SHorizontalBox::Slot()
+		.FillWidth(1.f)
+		.VAlign(VAlign_Center)
+		[
+			SNew(STextBlock)
+			.Text(EventsCategory.GetDisplayName())
+			.Font(IDetailLayoutBuilder::GetDetailFont())
+		]
+		+ SHorizontalBox::Slot()
+		.AutoWidth()
+		.HAlign(HAlign_Right)
+		.VAlign(VAlign_Center)
+		[
+			EventsHandle->CreateDefaultPropertyButtonWidgets()
+		];
 
-	return DebugControlsObject->DebuggedObjectWeak.Get();
-}
+	EventsCategory.HeaderContent(HeaderContentWidget, /*bWholeRowContent*/true);
 
-void FDebugControlsObjectDetails::CustomizeDebuggedObjectDetails(IDetailLayoutBuilder& InDetailBuilder)
-{
-	TArray<UObject*> DebuggedObjects;
-	DebuggedObjects.Reserve(DebugControls.Num());
+	TSharedRef<FDetailArrayBuilder> EventsArrayBuilder = MakeShared<FDetailArrayBuilder>(EventsHandle
+		, /*GenerateHeader*/false
+		, /*DisplayResetToDefault*/false
+		, /*bDisplayElementNum*/false);
 
-	// Gather all the Debugged objects within the debug controls
-	Algo::Transform(DebugControls, DebuggedObjects,
-		[](const TWeakObjectPtr<USceneStateDebugControlsObject>& InDebugControls)->UObject*
-		{
-			if (USceneStateDebugControlsObject* DebugControl = InDebugControls.Get())
-			{
-				return DebugControl->DebuggedObjectWeak.Get();
-			}
-			return nullptr;
-		});
+	TWeakPtr<IPropertyUtilities> PropertyUtilitiesWeak = InDetailBuilder.GetPropertyUtilities();
 
-	DebuggedObjects.RemoveAll(
-		[](UObject* InObject)
+	EventsArrayBuilder->OnGenerateArrayElementWidget(FOnGenerateArrayElementWidget::CreateLambda(
+		[PropertyUtilitiesWeak](TSharedRef<IPropertyHandle> InPropertyHandle, int32 InArrayIndex, IDetailChildrenBuilder& InChildrenBuilder)
 		{
-			return !InObject;
-		});
+			InChildrenBuilder.AddProperty(InPropertyHandle);
+		}));
 
-	InDetailBuilder
-		.EditCategory(TEXT("Debugged Object"))
-		.AddExternalObjects(DebuggedObjects, EPropertyLocation::Default,
-			FAddPropertyParams()
-			.HideRootObjectNode(/*bHideRootObjectNode*/true));
+	EventsCategory.AddCustomBuilder(EventsArrayBuilder);
 }
 
 } // UE::SceneState::Editor
diff --git a/Engine/Plugins/VirtualProduction/SceneState/Source/SceneStateBlueprintEditor/Private/DebugControls/SceneStateDebugControlsObjectDetails.h b/Engine/Plugins/VirtualProduction/SceneState/Source/SceneStateBlueprintEditor/Private/DebugControls/SceneStateDebugControlsObjectDetails.h
--- a/Engine/Plugins/VirtualProduction/SceneState/Source/SceneStateBlueprintEditor/Private/DebugControls/SceneStateDebugControlsObjectDetails.h
+++ b/Engine/Plugins/VirtualProduction/SceneState/Source/SceneStateBlueprintEditor/Private/DebugControls/SceneStateDebugControlsObjectDetails.h
@@ -31,14 +31,6 @@ public:
 private:
 	void CustomizeEventDetails(IDetailLayoutBuilder& InDetailBuilder);
 
-	/** Gets the scene state object within the debug controls object found at the given index */
-	USceneStateObject* GetDebuggedObject(int32 InIndex) const;
-
-	/** Adds the debugged object details (if there's only one to see) */
-	void CustomizeDebuggedObjectDetails(IDetailLayoutBuilder& InDetailBuilder);
-
-	TSharedPtr<IPropertyHandle> EventsHandle;
-
 	TArray<TWeakObjectPtr<USceneStateDebugControlsObject>> DebugControls;
 };
 
diff --git a/Engine/Plugins/VirtualProduction/SceneState/Source/SceneStateBlueprintEditor/Private/DebugControls/SceneStateDebugControlsTool.cpp b/Engine/Plugins/VirtualProduction/SceneState/Source/SceneStateBlueprintEditor/Private/DebugControls/SceneStateDebugControlsTool.cpp
--- a/Engine/Plugins/VirtualProduction/SceneState/Source/SceneStateBlueprintEditor/Private/DebugControls/SceneStateDebugControlsTool.cpp
+++ b/Engine/Plugins/VirtualProduction/SceneState/Source/SceneStateBlueprintEditor/Private/DebugControls/SceneStateDebugControlsTool.cpp
@@ -42,11 +42,6 @@ void FDebugControlsTool::Tick(float InDeltaSeconds)
 	}
 }
 
-void FDebugControlsTool::UpdateDebuggedObject()
-{
-	DebugControlsObject->DebuggedObjectWeak = GetDebuggedObject();
-}
-
 UObject* FDebugControlsTool::GetDebugControlsObject() const
 {
 	return DebugControlsObject.Get();
diff --git a/Engine/Plugins/VirtualProduction/SceneState/Source/SceneStateBlueprintEditor/Private/DebugControls/SceneStateDebugControlsTool.h b/Engine/Plugins/VirtualProduction/SceneState/Source/SceneStateBlueprintEditor/Private/DebugControls/SceneStateDebugControlsTool.h
--- a/Engine/Plugins/VirtualProduction/SceneState/Source/SceneStateBlueprintEditor/Private/DebugControls/SceneStateDebugControlsTool.h
+++ b/Engine/Plugins/VirtualProduction/SceneState/Source/SceneStateBlueprintEditor/Private/DebugControls/SceneStateDebugControlsTool.h
@@ -47,16 +47,13 @@ public:
 	/** Called to update the active debug executor if any */
 	void Tick(float InDeltaSeconds);
 
-	/** Called when the debugged object has changed */
-	void UpdateDebuggedObject();
-
 	/** Gets the debug controls object to show in details view */
 	UObject* GetDebugControlsObject() const;
 
-private:
 	/** Retrieves the currently debugged object */
 	USceneStateObject* GetDebuggedObject() const;
 
+private:
 	/** Gets the single selected state machine node, or null if nothing is selected or multiple items are selected */
 	const USceneStateMachineNode* GetSingleSelectedNode() const;
 
diff --git a/Engine/Plugins/VirtualProduction/SceneState/Source/SceneStateBlueprintEditor/Private/SceneStateBlueprintEditor.cpp b/Engine/Plugins/VirtualProduction/SceneState/Source/SceneStateBlueprintEditor/Private/SceneStateBlueprintEditor.cpp
--- a/Engine/Plugins/VirtualProduction/SceneState/Source/SceneStateBlueprintEditor/Private/SceneStateBlueprintEditor.cpp
+++ b/Engine/Plugins/VirtualProduction/SceneState/Source/SceneStateBlueprintEditor/Private/SceneStateBlueprintEditor.cpp
@@ -106,6 +106,28 @@ void FSceneStateBlueprintEditor::JumpToHyperlink(const UObject* InObjectReferenc
 	}
 }
 
+bool FSceneStateBlueprintEditor::InEditingMode() const
+{
+	// Always in editing mode when not in PIE
+	if (!IsPlayInEditorActive())
+	{
+		return true;
+	}
+
+	// Never in editing mode when debugging in bp
+	if (FSlateApplication::Get().InKismetDebuggingMode())
+	{
+		return false;
+	}
+
+	// By default, disallow editing while in PIE unless the blueprint allows it. 
+	if (const UBlueprint* Blueprint = GetBlueprintObj())
+	{
+		return Blueprint->CanAlwaysRecompileWhilePlayingInEditor();
+	}
+	return false;
+}
+
 void FSceneStateBlueprintEditor::CreateDefaultTabContents(const TArray<UBlueprint*>& InBlueprints)
 {
 	Super::CreateDefaultTabContents(InBlueprints);
diff --git a/Engine/Plugins/VirtualProduction/SceneState/Source/SceneStateBlueprintEditor/Private/SceneStateBlueprintEditor.h b/Engine/Plugins/VirtualProduction/SceneState/Source/SceneStateBlueprintEditor/Private/SceneStateBlueprintEditor.h
--- a/Engine/Plugins/VirtualProduction/SceneState/Source/SceneStateBlueprintEditor/Private/SceneStateBlueprintEditor.h
+++ b/Engine/Plugins/VirtualProduction/SceneState/Source/SceneStateBlueprintEditor/Private/SceneStateBlueprintEditor.h
@@ -39,6 +39,7 @@ public:
 
 	//~ Begin FBlueprintEditor
 	using Super::InitBlueprintEditor;
+	virtual bool InEditingMode() const;
 	virtual void CreateDefaultTabContents(const TArray<UBlueprint*>& InBlueprints) override;
 	virtual void CreateDefaultCommands() override;
 	virtual void RegisterApplicationModes(const TArray<UBlueprint*>& InBlueprints, bool bInShouldOpenInDefaultsMode, bool bInNewlyCreated) override;
