Motion Design Scene State: deprecated public access to state instance status + minor cleanup of functions in FSceneState no longer requiring state instance

#Synced-CL 49389891

diff --git a/Engine/Plugins/VirtualProduction/SceneState/Source/SceneState/Private/SceneState.cpp b/Engine/Plugins/VirtualProduction/SceneState/Source/SceneState/Private/SceneState.cpp
--- a/Engine/Plugins/VirtualProduction/SceneState/Source/SceneState/Private/SceneState.cpp
+++ b/Engine/Plugins/VirtualProduction/SceneState/Source/SceneState/Private/SceneState.cpp
@@ -25,16 +25,16 @@ void FSceneState::Enter(const FSceneStateExecutionContext& InContext) const
 	}
 
 	// Enter State, add a State Instance if not already present
-	FSceneStateInstance* Instance = InContext.FindOrAddStateInstance(*this);
-	if (!Instance || Instance->Status == EExecutionStatus::Running)
+	FSceneStateInstance& Instance = InContext.FindOrAddStateInstance(*this);
+	if (Instance.GetStatus() == EExecutionStatus::Running)
 	{
 		return;
 	}
 
 	UE_LOG(LogSceneState, Verbose, TEXT("State (%s) receiving enter"), GetStateName(InContext).GetData());
 
-	Instance->ElapsedTime = 0.f;
-	Instance->Status = EExecutionStatus::Running;
+	Instance.ElapsedTime = 0.f;
+	InContext.SetStateStatus(*this, Instance, EExecutionStatus::Running);
 
 	// Apply Event Handlers before anything else starts so the Event Data becomes available to substate machines and tasks
 	CaptureEvents(InContext);
@@ -52,7 +52,7 @@ void FSceneState::Enter(const FSceneStateExecutionContext& InContext) const
 
 	// Setup each Task
 	InContext.ForEachTask(*this,
-		[&InContext, Instance](const FSceneStateTask& InTask, FStructView InTaskInstance)->EIterationResult
+		[&InContext](const FSceneStateTask& InTask, FStructView InTaskInstance)->EIterationResult
 		{
 			InTask.Setup(InContext, InTaskInstance);
 			return EIterationResult::Continue;
@@ -63,7 +63,7 @@ void FSceneState::Enter(const FSceneStateExecutionContext& InContext) const
 		StateMachine.Start(InContext);
 	}
 
-	UpdateActiveTasks(InContext, *Instance);
+	UpdateActiveTasks(InContext);
 }
 
 void FSceneState::Tick(const FSceneStateExecutionContext& InContext, float InDeltaSeconds) const
@@ -77,7 +77,7 @@ void FSceneState::Tick(const FSceneStateExecutionContext& InContext, float InDel
 	}
 
 	FSceneStateInstance* Instance = InContext.FindStateInstance(*this);
-	if (!Instance || Instance->Status != EExecutionStatus::Running)
+	if (!Instance || Instance->GetStatus() != EExecutionStatus::Running)
 	{
 		return;
 	}
@@ -86,7 +86,7 @@ void FSceneState::Tick(const FSceneStateExecutionContext& InContext, float InDel
 
 	Instance->ElapsedTime += InDeltaSeconds;
 
-	UpdateActiveTasks(InContext, *Instance);
+	UpdateActiveTasks(InContext);
 
 	for (const FSceneStateMachine& StateMachine : InContext.GetStateMachines(*this))
 	{
@@ -112,7 +112,7 @@ void FSceneState::Exit(const FSceneStateExecutionContext& InContext) const
 	}
 
 	FSceneStateInstance* Instance = InContext.FindStateInstance(*this);
-	if (!Instance || Instance->Status != EExecutionStatus::Running)
+	if (!Instance || Instance->GetStatus() != EExecutionStatus::Running)
 	{
 		return;
 	}
@@ -139,7 +139,7 @@ void FSceneState::Exit(const FSceneStateExecutionContext& InContext) const
 		Transition.Exit(InContext);
 	}
 
-	Instance->Status = EExecutionStatus::Finished;
+	InContext.SetStateStatus(*this, *Instance, EExecutionStatus::Finished);
 	Instance->ElapsedTime = 0.f;
 
 	InContext.RemoveStateInstance(*this);
@@ -148,7 +148,7 @@ void FSceneState::Exit(const FSceneStateExecutionContext& InContext) const
 	ResetCapturedEvents(InContext);
 }
 
-void FSceneState::UpdateActiveTasks(const FSceneStateExecutionContext& InContext, FSceneStateInstance& InInstance) const
+void FSceneState::UpdateActiveTasks(const FSceneStateExecutionContext& InContext) const
 {
 	using namespace UE::SceneState;
 
@@ -170,7 +170,7 @@ void FSceneState::UpdateActiveTasks(const FSceneStateExecutionContext& InContext
 
 	// Start all the Tasks that haven't started yet and that meet their pre-requisites
 	InContext.ForEachTask(*this,
-		[&InContext, &InInstance, &FindTaskInstance](const FSceneStateTask& InTask, FStructView InTaskInstance)->EIterationResult
+		[&InContext, &FindTaskInstance](const FSceneStateTask& InTask, FStructView InTaskInstance)->EIterationResult
 		{
 			FSceneStateTaskInstance* TaskInstance = InTaskInstance.GetPtr<FSceneStateTaskInstance>();
 			if (!TaskInstance || TaskInstance->GetStatus() != EExecutionStatus::NotStarted)
diff --git a/Engine/Plugins/VirtualProduction/SceneState/Source/SceneState/Private/SceneStateExecutionContext.cpp b/Engine/Plugins/VirtualProduction/SceneState/Source/SceneState/Private/SceneStateExecutionContext.cpp
--- a/Engine/Plugins/VirtualProduction/SceneState/Source/SceneState/Private/SceneStateExecutionContext.cpp
+++ b/Engine/Plugins/VirtualProduction/SceneState/Source/SceneState/Private/SceneStateExecutionContext.cpp
@@ -477,9 +477,9 @@ const FSceneStateConduit* FSceneStateExecutionContext::GetConduit(const FSceneSt
 	return &ConduitsView[AbsoluteIndex];
 }
 
-FSceneStateInstance* FSceneStateExecutionContext::FindOrAddStateInstance(const FSceneState& InState) const
+FSceneStateInstance& FSceneStateExecutionContext::FindOrAddStateInstance(const FSceneState& InState) const
 {
-	return &StateInstances.FindOrAdd(GetStateIndex(InState));
+	return StateInstances.FindOrAdd(GetStateIndex(InState));
 }
 
 FSceneStateInstance* FSceneStateExecutionContext::FindStateInstance(uint16 InAbsoluteIndex) const
@@ -714,6 +714,16 @@ FStructView FSceneStateExecutionContext::FindFunctionInstance(uint16 InAbsoluteI
 	return FunctionInstanceViews.FindRef(InAbsoluteIndex);
 }
 
+void FSceneStateExecutionContext::SetStateStatus(const FSceneState& InState, FSceneStateInstance& InStateInstance, UE::SceneState::EExecutionStatus InStatus) const
+{
+	if (InStateInstance.GetStatus() == InStatus)
+	{
+		return;
+	}
+
+	InStateInstance.SetStatus(InStatus);
+}
+
 const FSceneStateExecutionContext* FSceneStateExecutionContext::FindParentContext() const
 {
 	if (TSharedPtr<UE::SceneState::FExecutionContextRegistry> ContextRegistry = ContextRegistryWeak.Pin())
diff --git a/Engine/Plugins/VirtualProduction/SceneState/Source/SceneState/Private/SceneStateInstance.cpp b/Engine/Plugins/VirtualProduction/SceneState/Source/SceneState/Private/SceneStateInstance.cpp
--- a/Engine/Plugins/VirtualProduction/SceneState/Source/SceneState/Private/SceneStateInstance.cpp
+++ b/Engine/Plugins/VirtualProduction/SceneState/Source/SceneState/Private/SceneStateInstance.cpp
@@ -19,10 +19,26 @@ namespace UE::SceneState::Private
 
 } // UE::SceneState::Private
 
+PRAGMA_DISABLE_DEPRECATION_WARNINGS
 FSceneStateInstance::FSceneStateInstance()
 	: InstanceId(UE::SceneState::Private::GetNextInstanceId())
 {
 }
+PRAGMA_ENABLE_DEPRECATION_WARNINGS
+
+UE::SceneState::EExecutionStatus FSceneStateInstance::GetStatus() const
+{
+PRAGMA_DISABLE_DEPRECATION_WARNINGS
+	return Status;
+PRAGMA_ENABLE_DEPRECATION_WARNINGS
+}
+
+void FSceneStateInstance::SetStatus(UE::SceneState::EExecutionStatus InStatus)
+{
+PRAGMA_DISABLE_DEPRECATION_WARNINGS
+	Status = InStatus;
+PRAGMA_ENABLE_DEPRECATION_WARNINGS
+}
 
 uint16 FSceneStateInstance::GetInstanceId() const
 {
diff --git a/Engine/Plugins/VirtualProduction/SceneState/Source/SceneState/Private/SceneStateObject.cpp b/Engine/Plugins/VirtualProduction/SceneState/Source/SceneState/Private/SceneStateObject.cpp
--- a/Engine/Plugins/VirtualProduction/SceneState/Source/SceneState/Private/SceneStateObject.cpp
+++ b/Engine/Plugins/VirtualProduction/SceneState/Source/SceneState/Private/SceneStateObject.cpp
@@ -40,7 +40,7 @@ bool USceneStateObject::IsActive() const
 	if (const FSceneState* RootState = GetRootState())
 	{
 		const FSceneStateInstance* RootStateInstance = RootExecutionContext.FindStateInstance(*RootState);
-		return RootStateInstance && RootStateInstance->Status == UE::SceneState::EExecutionStatus::Running;
+		return RootStateInstance && RootStateInstance->GetStatus() == UE::SceneState::EExecutionStatus::Running;
 	}
 	return false;
 }
diff --git a/Engine/Plugins/VirtualProduction/SceneState/Source/SceneState/Public/SceneState.h b/Engine/Plugins/VirtualProduction/SceneState/Source/SceneState/Public/SceneState.h
--- a/Engine/Plugins/VirtualProduction/SceneState/Source/SceneState/Public/SceneState.h
+++ b/Engine/Plugins/VirtualProduction/SceneState/Source/SceneState/Public/SceneState.h
@@ -7,6 +7,8 @@
 #include "SceneStateReentryGuard.h"
 #include "SceneState.generated.h"
 
+#define UE_API SCENESTATE_API
+
 struct FConstStructView;
 struct FSceneStateExecutionContext;
 struct FSceneStateInstance;
@@ -42,13 +44,13 @@ struct FSceneState
 	GENERATED_BODY()
 
 	/** Called the state becomes the active state in its parent state machine */
-	SCENESTATE_API void Enter(const FSceneStateExecutionContext& InContext) const;
+	UE_API void Enter(const FSceneStateExecutionContext& InContext) const;
 
 	/** Called on every tick while the state is active */
-	SCENESTATE_API void Tick(const FSceneStateExecutionContext& InContext, float InDeltaSeconds) const;
+	UE_API void Tick(const FSceneStateExecutionContext& InContext, float InDeltaSeconds) const;
 
 	/** Called when the state is no longer active in its parent state machine */
-	SCENESTATE_API void Exit(const FSceneStateExecutionContext& InContext) const;
+	UE_API void Exit(const FSceneStateExecutionContext& InContext) const;
 
 	FSceneStateRange GetTaskRange() const
 	{
@@ -71,20 +73,20 @@ struct FSceneState
 	}
 
 	/** Starts pending tasks whose prerequisites are satisfied */
-	void UpdateActiveTasks(const FSceneStateExecutionContext& InContext, FSceneStateInstance& InInstance) const;
+	void UpdateActiveTasks(const FSceneStateExecutionContext& InContext) const;
 
 	/** Returns the State Name used for logging. Returns an empty string view if metadata is not available */
 	FStringView GetStateName(const FSceneStateExecutionContext& InContext) const;
 
 	/** Returns whether there's any Task pending to finish */
-	SCENESTATE_API bool HasPendingTasks(const FSceneStateExecutionContext& InContext) const;
+	UE_API bool HasPendingTasks(const FSceneStateExecutionContext& InContext) const;
 
 	/** Called on start to create the required task instances that will run for this state */
-	SCENESTATE_API void AllocateTaskInstances(const FSceneStateExecutionContext& InContext, TConstArrayView<FConstStructView> InTemplateTaskInstances) const;
+	UE_API void AllocateTaskInstances(const FSceneStateExecutionContext& InContext, TConstArrayView<FConstStructView> InTemplateTaskInstances) const;
 
 private:
 	/** Instances the instanced references in source data into the target data using the provided duplication functor */
-	SCENESTATE_API void InstanceTaskObjects(UObject* InOuter
+	UE_API void InstanceTaskObjects(UObject* InOuter
 		, TConstArrayView<FStructView> InTargets
 		, TConstArrayView<FConstStructView> InSources
 		, TFunctionRef<UObject*(FObjectDuplicationParameters&)> InDuplicationFunc) const;
@@ -120,3 +122,5 @@ private:
 	friend UE::SceneState::Editor::FBlueprintCompilerContext;
 	friend UE::SceneState::Editor::FStateMachineCompiler;
 };
+
+#undef UE_API
diff --git a/Engine/Plugins/VirtualProduction/SceneState/Source/SceneState/Public/SceneStateExecutionContext.h b/Engine/Plugins/VirtualProduction/SceneState/Source/SceneState/Public/SceneStateExecutionContext.h
--- a/Engine/Plugins/VirtualProduction/SceneState/Source/SceneState/Public/SceneStateExecutionContext.h
+++ b/Engine/Plugins/VirtualProduction/SceneState/Source/SceneState/Public/SceneStateExecutionContext.h
@@ -191,7 +191,7 @@ struct FSceneStateExecutionContext
 	const FSceneStateConduit* GetConduit(const FSceneStateMachine& InStateMachine, uint16 InRelativeIndex) const;
 
 	/** Returns the State Instance of this context for the given state, adding a new one if it's not found */
-	FSceneStateInstance* FindOrAddStateInstance(const FSceneState& InState) const;
+	FSceneStateInstance& FindOrAddStateInstance(const FSceneState& InState) const;
 
 	/** Returns the existing State Instance of this context for the given state absolute index. Null if not found */
 	FSceneStateInstance* FindStateInstance(uint16 InAbsoluteIndex) const;
@@ -265,6 +265,9 @@ struct FSceneStateExecutionContext
 	/** Returns the function instance of this context for the given function absolute index. Returns an invalid view if not found */
 	FStructView FindFunctionInstance(uint16 InAbsoluteIndex) const;
 
+	/** Called to set the status of the given state */
+	void SetStateStatus(const FSceneState& InState, FSceneStateInstance& InStateInstance, UE::SceneState::EExecutionStatus InStatus) const;
+
 private:
 	/** Retrieves the parent execution context if existing. Can be null */
 	const FSceneStateExecutionContext* FindParentContext() const;
diff --git a/Engine/Plugins/VirtualProduction/SceneState/Source/SceneState/Public/SceneStateInstance.h b/Engine/Plugins/VirtualProduction/SceneState/Source/SceneState/Public/SceneStateInstance.h
--- a/Engine/Plugins/VirtualProduction/SceneState/Source/SceneState/Public/SceneStateInstance.h
+++ b/Engine/Plugins/VirtualProduction/SceneState/Source/SceneState/Public/SceneStateInstance.h
@@ -5,23 +5,41 @@
 #include "SceneStateEnums.h"
 #include "SceneStateInstance.generated.h"
 
+#define UE_API SCENESTATE_API
+
 /** Instance data of a State */
 USTRUCT()
 struct FSceneStateInstance
 {
 	GENERATED_BODY()
 
-	SCENESTATE_API FSceneStateInstance();
+PRAGMA_DISABLE_DEPRECATION_WARNINGS
+	UE_API FSceneStateInstance();
+	FSceneStateInstance(const FSceneStateInstance&) = default;
+	FSceneStateInstance(FSceneStateInstance&&) = default;
+	FSceneStateInstance& operator=(const FSceneStateInstance&) = default;
+	FSceneStateInstance& operator=(FSceneStateInstance&&) = default;
+PRAGMA_ENABLE_DEPRECATION_WARNINGS
 
 	/** Time elapsed since the state entered */
 	float ElapsedTime = 0.f;
 
+	/** Gets the execution status of the state instance */
+	UE_API UE::SceneState::EExecutionStatus GetStatus() const;
+
+	/** Sets the execution status of the state instance */
+	void SetStatus(UE::SceneState::EExecutionStatus InStatus);
+
 	/** Current status of this state instance */
+	UE_DEPRECATED(5.8, "Public access to Status is deprecated, and it will become private in a future release. Please update code to use the setter/getter instead.")
 	UE::SceneState::EExecutionStatus Status = UE::SceneState::EExecutionStatus::NotStarted;
 
+	/** Get an identifier for this instance */
 	uint16 GetInstanceId() const;
 
 private:
 	/** The id for this instance. This is used to differentiate state instances for a same state and consequently task instances for a same task*/
 	uint16 InstanceId;
 };
+
+#undef UE_API
