Motion Design: added trivial relocability trait to execution context to indicate it's not trivially relocatable

#Synced-CL 49564377

diff --git a/Engine/Plugins/VirtualProduction/SceneState/Source/SceneState/Public/SceneStateExecutionContext.h b/Engine/Plugins/VirtualProduction/SceneState/Source/SceneState/Public/SceneStateExecutionContext.h
--- a/Engine/Plugins/VirtualProduction/SceneState/Source/SceneState/Public/SceneStateExecutionContext.h
+++ b/Engine/Plugins/VirtualProduction/SceneState/Source/SceneState/Public/SceneStateExecutionContext.h
@@ -365,3 +365,10 @@ private:
 };
 
 #undef UE_API
+
+/** FSceneStateExecutionContext is not trivially relocatable */
+template<UE::CDerivedFrom<FSceneStateExecutionContext> T>
+struct TIsTriviallyRelocatable<T>
+{
+	static inline constexpr bool Value = false;
+};
