Motion Design Scene State: Improved scenestate debug execution flow visualization by tracking state execution without relying on state instances (as these get removed as soon as states exit). This is so that states that  have exited still have time to be visualized and fade out.
In addition, nodes (e.g. State,Task,Graph) store their last compiled index so that these get saved and don't have to rely on a recompilation since TemplateData does not serialize the maps of node to index.

#Synced-CL 50160130

diff --git a/Engine/Plugins/VirtualProduction/SceneState/Source/SceneState/Internal/Debug/SceneStateDebugStateContext.h b/Engine/Plugins/VirtualProduction/SceneState/Source/SceneState/Internal/Debug/SceneStateDebugStateContext.h
new file mode 100644
--- /dev/null
+++ b/Engine/Plugins/VirtualProduction/SceneState/Source/SceneState/Internal/Debug/SceneStateDebugStateContext.h
@@ -0,0 +1,42 @@
+// Copyright Epic Games, Inc. All Rights Reserved.
+
+#pragma once
+
+#if WITH_EDITOR
+#include "SceneStateExecutionContextHandle.h"
+#include "Templates/SharedPointer.h"
+
+#define UE_API SCENESTATE_API
+
+struct FSceneState;
+struct FSceneStateExecutionContext;
+struct FSceneStateInstance;
+
+namespace UE::SceneState
+{
+
+/** Context for a debugged state */
+struct FDebugStateContext
+{
+	explicit FDebugStateContext(uint16 InStateIndex, const FSceneStateInstance& InStateInstance);
+
+	uint16 GetStateIndex() const
+	{
+		return StateIndex;
+	}
+
+	bool Equals(const FDebugStateContext& InOther) const;
+
+private:
+	/** Index of the state that is being debugged */
+	uint16 StateIndex;
+
+	/** Identifier for a single state instance */
+	uint16 StateInstanceId;
+};
+
+} // UE::SceneState
+
+#undef UE_API
+
+#endif
diff --git a/Engine/Plugins/VirtualProduction/SceneState/Source/SceneState/Internal/Debug/SceneStateDebugStateInstance.h b/Engine/Plugins/VirtualProduction/SceneState/Source/SceneState/Internal/Debug/SceneStateDebugStateInstance.h
new file mode 100644
--- /dev/null
+++ b/Engine/Plugins/VirtualProduction/SceneState/Source/SceneState/Internal/Debug/SceneStateDebugStateInstance.h
@@ -0,0 +1,66 @@
+// Copyright Epic Games, Inc. All Rights Reserved.
+
+#pragma once
+
+#if WITH_EDITOR
+#include "SceneStateDebugStateContext.h"
+#include "Misc/Optional.h"
+#include "SceneStateEnums.h"
+
+#define UE_API SCENESTATE_API
+
+struct FSceneStateInstance;
+
+namespace UE::SceneState
+{
+
+/** Holds debug information about the state to visualize as node info popups in the state machine graph */
+class FDebugStateInstance
+{
+	static constexpr double EnteringTime = 0.125;
+	static constexpr double ExitingTime = 0.75;
+
+public:
+	explicit FDebugStateInstance(const FDebugStateContext& InContext, EExecutionStatus InInitialStatus);
+
+	uint16 GetStateIndex() const;
+
+	/** Returns true if the debug instance matches the given state context information */
+	bool operator==(const FDebugStateContext& InOther) const;
+
+	/** Calculate how 'active' a state is. Returns a float value in the range [0, 1] */
+	UE_API float GetActivePercentage() const;
+
+	/** Determines whether the state is in a current valid state */
+	UE_API bool IsValid() const;
+
+	/**
+	 * Returns the elapsed time for the current execution status.
+	 * The elapsed time resets every time the execution status changes.
+	 */
+	UE_API double GetStatusElapsedTime() const;
+
+	EExecutionStatus GetExecutionStatus() const
+	{
+		return ExecutionStatus;
+	}
+
+	/** Sets the execution status of this state debug instance (to match with the actual executed state instance)  */
+	void SetExecutionStatus(EExecutionStatus InExecutionStatus);
+
+private:
+	/** Context of the executed state */
+	FDebugStateContext DebugStateContext;
+
+	/** Current execution status of the state */
+	EExecutionStatus ExecutionStatus = EExecutionStatus::NotStarted;
+
+	/** Last time the status was changed */
+	double StatusChangeTime = 0.0;
+};
+
+} // UE::SceneState
+
+#undef UE_API
+
+#endif // WITH_EDITOR
diff --git a/Engine/Plugins/VirtualProduction/SceneState/Source/SceneState/Internal/Debug/SceneStateDebugger.h b/Engine/Plugins/VirtualProduction/SceneState/Source/SceneState/Internal/Debug/SceneStateDebugger.h
new file mode 100644
--- /dev/null
+++ b/Engine/Plugins/VirtualProduction/SceneState/Source/SceneState/Internal/Debug/SceneStateDebugger.h
@@ -0,0 +1,63 @@
+// Copyright Epic Games, Inc. All Rights Reserved.
+
+#pragma once
+
+#if WITH_EDITOR
+#include "Containers/Map.h"
+#include "SceneStateDebugStateInstance.h"
+#include "Templates/SharedPointer.h"
+#include "UObject/WeakInterfacePtr.h"
+#include "UObject/WeakObjectPtrTemplates.h"
+
+#define UE_API SCENESTATE_API
+
+class USceneStateObject;
+
+namespace UE::SceneState
+{
+	struct FDebugStateContext;
+}
+
+namespace UE::SceneState
+{
+
+/** Processes execution events like state status changes. Editor-only */
+class FDebugger : public TSharedFromThis<FDebugger>
+{
+public:
+	UE_API ~FDebugger();
+
+	/** Called to attach this debugger to the target scene state object */
+	UE_API void AttachTo(USceneStateObject* InObject);
+
+	/** Called to detach this debugger from any existing scene state object */
+	void Detach();
+
+	/** Returns whether this debugger is attached to the given object */
+	bool IsAttachedTo(USceneStateObject* InObject) const;
+
+	/** Calls the given functor for every state debug instance that matches the state index */
+	UE_API void ForEachDebugInstanceOfState(uint16 InStateIndex, TFunctionRef<void(const FDebugStateInstance&)> InFunc);
+
+	/** Called when a state instance has changed status */
+	void NotifyStateStatusChanged(const FDebugStateContext& InContext, UE::SceneState::EExecutionStatus InStatus);
+
+private:
+	/** Populates debug state instances with currently active states */
+	void Initialize();
+
+	/** Removes all state debug instances that are no longer valid */
+	void CleanInvalidDebugInstances();
+
+	/** All the relevant debug instances that are or were recently active. */
+	TArray<FDebugStateInstance> DebugStateInstances;
+
+	/** The object that this debugger is currently attached to */
+	TWeakObjectPtr<USceneStateObject> DebuggedObjectWeak;
+};
+
+} // UE::SceneState
+
+#undef UE_API
+
+#endif // WITH_EDITOR
diff --git a/Engine/Plugins/VirtualProduction/SceneState/Source/SceneState/Private/Debug/SceneStateDebugStateContext.cpp b/Engine/Plugins/VirtualProduction/SceneState/Source/SceneState/Private/Debug/SceneStateDebugStateContext.cpp
new file mode 100644
--- /dev/null
+++ b/Engine/Plugins/VirtualProduction/SceneState/Source/SceneState/Private/Debug/SceneStateDebugStateContext.cpp
@@ -0,0 +1,25 @@
+// Copyright Epic Games, Inc. All Rights Reserved.
+
+#if WITH_EDITOR
+#include "Debug/SceneStateDebugStateContext.h"
+#include "SceneStateExecutionContext.h"
+#include "SceneStateExecutionContextRegistry.h"
+
+namespace UE::SceneState
+{
+
+FDebugStateContext::FDebugStateContext(uint16 InStateIndex, const FSceneStateInstance& InStateInstance)
+	: StateIndex(InStateIndex)
+	, StateInstanceId(InStateInstance.GetInstanceId())
+{
+}
+
+bool FDebugStateContext::Equals(const FDebugStateContext& InOther) const
+{
+	return StateIndex == InOther.StateIndex
+		&& StateInstanceId == InOther.StateInstanceId;
+}
+
+} // UE::SceneState
+
+#endif // WITH_EDITOR
diff --git a/Engine/Plugins/VirtualProduction/SceneState/Source/SceneState/Private/Debug/SceneStateDebugStateInstance.cpp b/Engine/Plugins/VirtualProduction/SceneState/Source/SceneState/Private/Debug/SceneStateDebugStateInstance.cpp
new file mode 100644
--- /dev/null
+++ b/Engine/Plugins/VirtualProduction/SceneState/Source/SceneState/Private/Debug/SceneStateDebugStateInstance.cpp
@@ -0,0 +1,70 @@
+// Copyright Epic Games, Inc. All Rights Reserved.
+
+#if WITH_EDITOR
+#include "Debug/SceneStateDebugStateInstance.h"
+#include "Debug/SceneStateDebugStateContext.h"
+#include "Misc/App.h"
+#include "Misc/AppTime.h"
+#include "SceneStateExecutionContext.h"
+
+namespace UE::SceneState
+{
+
+FDebugStateInstance::FDebugStateInstance(const FDebugStateContext& InContext, EExecutionStatus InInitialStatus)
+	: DebugStateContext(InContext)
+{
+	SetExecutionStatus(InInitialStatus);
+}
+
+uint16 FDebugStateInstance::GetStateIndex() const
+{
+	return DebugStateContext.GetStateIndex();
+}
+
+bool FDebugStateInstance::operator==(const FDebugStateContext& InOther) const
+{
+	return DebugStateContext.Equals(InOther);
+}
+
+float FDebugStateInstance::GetActivePercentage() const
+{
+	const double StatusElapsedTime = GetStatusElapsedTime();
+
+	const double StatusChangeTotalTime = ExecutionStatus == EExecutionStatus::Running
+		? EnteringTime
+		: ExitingTime;
+
+	const float Percentage = StatusElapsedTime >= StatusChangeTotalTime
+		? 1.f
+		: FMath::Clamp(static_cast<float>(StatusElapsedTime / StatusChangeTotalTime), 0.f, 1.f);
+
+	return ExecutionStatus == EExecutionStatus::Running
+		? Percentage
+		: 1.f - Percentage;
+}
+
+bool FDebugStateInstance::IsValid() const
+{
+	if (ExecutionStatus == EExecutionStatus::Finished)
+	{
+		return GetStatusElapsedTime() <= ExitingTime;
+	}
+	return ExecutionStatus == EExecutionStatus::Running;
+}
+
+double FDebugStateInstance::GetStatusElapsedTime() const
+{
+	return FAppTime::GetCurrentTime() - StatusChangeTime;
+}
+
+void FDebugStateInstance::SetExecutionStatus(EExecutionStatus InExecutionStatus)
+{
+	if (ExecutionStatus != InExecutionStatus)
+	{
+		ExecutionStatus = InExecutionStatus;
+		StatusChangeTime = FAppTime::GetCurrentTime();
+	}
+}
+
+} // UE::SceneState
+#endif
diff --git a/Engine/Plugins/VirtualProduction/SceneState/Source/SceneState/Private/Debug/SceneStateDebugger.cpp b/Engine/Plugins/VirtualProduction/SceneState/Source/SceneState/Private/Debug/SceneStateDebugger.cpp
new file mode 100644
--- /dev/null
+++ b/Engine/Plugins/VirtualProduction/SceneState/Source/SceneState/Private/Debug/SceneStateDebugger.cpp
@@ -0,0 +1,124 @@
+// Copyright Epic Games, Inc. All Rights Reserved.
+
+#if WITH_EDITOR
+#include "Debug/SceneStateDebugger.h"
+#include "Debug/SceneStateDebugStateContext.h"
+#include "SceneStateExecutionContextRegistry.h"
+#include "SceneStateInstance.h"
+#include "SceneStateObject.h"
+
+namespace UE::SceneState
+{
+
+FDebugger::~FDebugger()
+{
+	Detach();
+}
+
+void FDebugger::AttachTo(USceneStateObject* InObject)
+{
+	if (DebuggedObjectWeak != InObject)
+	{
+		Detach();
+
+		DebuggedObjectWeak = InObject;
+
+		if (InObject)
+		{
+			InObject->SetDebugger(AsWeak());
+		}
+
+		Initialize();
+	}
+}
+
+void FDebugger::Detach()
+{
+	if (!UObjectInitialized())
+	{
+		DebuggedObjectWeak.Reset();
+		DebugStateInstances.Empty();
+		return;
+	}
+
+	USceneStateObject* const DebuggedObject = DebuggedObjectWeak.Get();
+	DebuggedObjectWeak.Reset(); // reset to prevent loop due to reentry via SetDebugger(nullptr)
+	DebugStateInstances.Empty();
+
+	if (DebuggedObject)
+	{
+		const UE::SceneState::FDebugger* const Debugger = DebuggedObject->GetDebugger().Get();
+		if (ensure(!Debugger || Debugger == this))
+		{
+			DebuggedObject->SetDebugger(nullptr);
+		}
+	}
+}
+
+bool FDebugger::IsAttachedTo(USceneStateObject* InObject) const
+{
+	return InObject && DebuggedObjectWeak == InObject;
+}
+
+void FDebugger::ForEachDebugInstanceOfState(uint16 InStateIndex, TFunctionRef<void(const FDebugStateInstance&)> InFunc)
+{
+	CleanInvalidDebugInstances();
+
+	for (const FDebugStateInstance& StateDebugInstance : DebugStateInstances)
+	{
+		if (StateDebugInstance.GetStateIndex() == InStateIndex)
+		{
+			InFunc(StateDebugInstance);
+		}
+	}
+}
+
+void FDebugger::NotifyStateStatusChanged(const FDebugStateContext& InContext, UE::SceneState::EExecutionStatus InStatus)
+{
+	CleanInvalidDebugInstances();
+
+	if (FDebugStateInstance* DebugStateInstance = DebugStateInstances.FindByKey(InContext))
+	{
+		DebugStateInstance->SetExecutionStatus(InStatus);
+	}
+	else
+	{
+		// Status is already handled within debug state instance construction
+		DebugStateInstances.Emplace(InContext, InStatus);
+	}
+}
+
+void FDebugger::Initialize()
+{
+	USceneStateObject* const DebuggedObject = DebuggedObjectWeak.Get();
+	if (!DebuggedObject)
+	{
+		return;
+	}
+
+	DebugStateInstances.Empty();
+
+	DebuggedObject->GetContextRegistry()->ForEachExecutionContext(
+		[This=this](const FSceneStateExecutionContext& InContext)
+		{
+			InContext.ForEachStateInstance(
+				[&InContext, This](uint16 InStateIndex, const FSceneStateInstance& InStateInstance)
+				{
+					const FDebugStateContext StateContext(InStateIndex, InStateInstance);
+					This->DebugStateInstances.Emplace(StateContext, InStateInstance.GetStatus());
+				});
+		});
+}
+
+void FDebugger::CleanInvalidDebugInstances()
+{
+	DebugStateInstances.RemoveAll(
+		[](const FDebugStateInstance& InDebugInstance)
+		{
+			return !InDebugInstance.IsValid();
+		});
+}
+
+} // UE::SceneState
+
+#endif // WITH_EDITOR
diff --git a/Engine/Plugins/VirtualProduction/SceneState/Source/SceneState/Private/SceneState.cpp b/Engine/Plugins/VirtualProduction/SceneState/Source/SceneState/Private/SceneState.cpp
--- a/Engine/Plugins/VirtualProduction/SceneState/Source/SceneState/Private/SceneState.cpp
+++ b/Engine/Plugins/VirtualProduction/SceneState/Source/SceneState/Private/SceneState.cpp
@@ -33,7 +33,6 @@ void FSceneState::Enter(const FSceneStateExecutionContext& InContext) const
 
 	UE_LOG(LogSceneState, Verbose, TEXT("State (%s) receiving enter"), GetStateName(InContext).GetData());
 
-	Instance.ElapsedTime = 0.f;
 	InContext.SetStateStatus(*this, Instance, EExecutionStatus::Running);
 
 	// Apply Event Handlers before anything else starts so the Event Data becomes available to substate machines and tasks
@@ -84,8 +83,6 @@ void FSceneState::Tick(const FSceneStateExecutionContext& InContext, float InDel
 
 	UE_LOG(LogSceneState, VeryVerbose, TEXT("State (%s) receiving tick"), GetStateName(InContext).GetData());
 
-	Instance->ElapsedTime += InDeltaSeconds;
-
 	UpdateActiveTasks(InContext);
 
 	for (const FSceneStateMachine& StateMachine : InContext.GetStateMachines(*this))
@@ -140,7 +137,6 @@ void FSceneState::Exit(const FSceneStateExecutionContext& InContext) const
 	}
 
 	InContext.SetStateStatus(*this, *Instance, EExecutionStatus::Finished);
-	Instance->ElapsedTime = 0.f;
 
 	InContext.RemoveStateInstance(*this);
 	InContext.RemoveTaskInstanceContainer(*this);
diff --git a/Engine/Plugins/VirtualProduction/SceneState/Source/SceneState/Private/SceneStateExecutionContext.cpp b/Engine/Plugins/VirtualProduction/SceneState/Source/SceneState/Private/SceneStateExecutionContext.cpp
--- a/Engine/Plugins/VirtualProduction/SceneState/Source/SceneState/Private/SceneStateExecutionContext.cpp
+++ b/Engine/Plugins/VirtualProduction/SceneState/Source/SceneState/Private/SceneStateExecutionContext.cpp
@@ -15,6 +15,11 @@
 #include "Tasks/SceneStateTaskBindingExtension.h"
 #include "Tasks/SceneStateTaskInstance.h"
 
+#if WITH_EDITOR
+#include "Debug/SceneStateDebugStateContext.h"
+#include "Debug/SceneStateDebugger.h"
+#endif
+
 const FSceneStateExecutionContext FSceneStateExecutionContext::InvalidContext;
 
 FSceneStateExecutionContext::~FSceneStateExecutionContext()
@@ -22,9 +27,14 @@ FSceneStateExecutionContext::~FSceneStateExecutionContext()
 	Reset();
 }
 
-TSharedPtr<const UE::SceneState::FExecutionContextRegistry> FSceneStateExecutionContext::GetContextRegistry() const
+TWeakPtr<const UE::SceneState::FExecutionContextRegistry> FSceneStateExecutionContext::GetContextRegistry() const
+{
+	return ContextRegistryWeak;
+}
+
+UE::SceneState::FExecutionContextHandle FSceneStateExecutionContext::GetContextHandle() const
 {
-	return ContextRegistryWeak.Pin();
+	return ContextHandle;
 }
 
 const FSceneStateBindingCollection& FSceneStateExecutionContext::GetBindingCollection() const
@@ -37,7 +47,7 @@ const FString& FSceneStateExecutionContext::GetExecutionContextName() const
 	return ContextName;
 }
 
-UObject* FSceneStateExecutionContext::GetRootObject() const
+USceneStateObject* FSceneStateExecutionContext::GetRootObject() const
 {
 	return RootObject;
 }
@@ -98,7 +108,11 @@ void FSceneStateExecutionContext::Setup(const FSetupParams& InParams)
 
 	if (InParams.ContextRegistry.IsValid())
 	{
-		InParams.ContextRegistry->RegisterContext(this);
+		ContextHandle = InParams.ContextRegistry->RegisterContext(this);
+	}
+	else
+	{
+		ContextHandle = UE::SceneState::FExecutionContextHandle();
 	}
 }
 
@@ -120,6 +134,9 @@ void FSceneStateExecutionContext::Reset()
 		ContextRegistry->UnregisterContext(this);
 	}
 
+	ContextHandle = UE::SceneState::FExecutionContextHandle();
+	ParentHandle = UE::SceneState::FExecutionContextHandle();
+
 	ContextRegistryWeak.Reset();
 }
 
@@ -497,6 +514,14 @@ void FSceneStateExecutionContext::RemoveStateInstance(const FSceneState& InState
 	StateInstances.Remove(GetStateIndex(InState));
 }
 
+void FSceneStateExecutionContext::ForEachStateInstance(TFunctionRef<void(uint16, const FSceneStateInstance&)> InFunc) const
+{
+	for (const TPair<uint16, FSceneStateInstance>& StateInstance : StateInstances)
+	{
+		InFunc(StateInstance.Key, StateInstance.Value);
+	}
+}
+
 FConstStructView FSceneStateExecutionContext::FindTask(uint16 InAbsoluteIndex) const
 {
 	if (TemplateData->Tasks.IsValidIndex(InAbsoluteIndex))
@@ -722,6 +747,17 @@ void FSceneStateExecutionContext::SetStateStatus(const FSceneState& InState, FSc
 	}
 
 	InStateInstance.SetStatus(InStatus);
+
+#if WITH_EDITOR
+	if (RootObject)
+	{
+		if (TSharedPtr<UE::SceneState::FDebugger> Debugger = RootObject->GetDebugger())
+		{
+			const UE::SceneState::FDebugStateContext StateContext(GetStateIndex(InState), InStateInstance);
+			Debugger->NotifyStateStatusChanged(StateContext, InStatus);
+		}
+	}
+#endif
 }
 
 const FSceneStateExecutionContext* FSceneStateExecutionContext::FindParentContext() const
diff --git a/Engine/Plugins/VirtualProduction/SceneState/Source/SceneState/Private/SceneStateExecutionContextRegistry.cpp b/Engine/Plugins/VirtualProduction/SceneState/Source/SceneState/Private/SceneStateExecutionContextRegistry.cpp
--- a/Engine/Plugins/VirtualProduction/SceneState/Source/SceneState/Private/SceneStateExecutionContextRegistry.cpp
+++ b/Engine/Plugins/VirtualProduction/SceneState/Source/SceneState/Private/SceneStateExecutionContextRegistry.cpp
@@ -7,7 +7,7 @@
 namespace UE::SceneState
 {
 
-void FExecutionContextRegistry::RegisterContext(const FSceneStateExecutionContext* InContext)
+FExecutionContextHandle FExecutionContextRegistry::RegisterContext(const FSceneStateExecutionContext* InContext)
 {
 	check(InContext && !bIterating);
 
@@ -25,6 +25,7 @@ void FExecutionContextRegistry::RegisterContext(const FSceneStateExecutionContex
 	const int32 HandleIndex = Handles.Add(Handle);
 	const int32 ContextIndex = Contexts.Add(InContext);
 	check(HandleIndex == ContextIndex);
+	return Handle;
 }
 
 void FExecutionContextRegistry::UnregisterContext(const FSceneStateExecutionContext* InContext)
diff --git a/Engine/Plugins/VirtualProduction/SceneState/Source/SceneState/Private/SceneStateExecutionContextRegistry.h b/Engine/Plugins/VirtualProduction/SceneState/Source/SceneState/Private/SceneStateExecutionContextRegistry.h
--- a/Engine/Plugins/VirtualProduction/SceneState/Source/SceneState/Private/SceneStateExecutionContextRegistry.h
+++ b/Engine/Plugins/VirtualProduction/SceneState/Source/SceneState/Private/SceneStateExecutionContextRegistry.h
@@ -16,7 +16,7 @@ class FExecutionContextRegistry : public TSharedFromThis<FExecutionContextRegist
 {
 public:
 	/** Called on Context setup to add the given context to the registry */
-	void RegisterContext(const FSceneStateExecutionContext* InContext);
+	FExecutionContextHandle RegisterContext(const FSceneStateExecutionContext* InContext);
 
 	/** Called on Context destructor to remove the given context from the registry */
 	void UnregisterContext(const FSceneStateExecutionContext* InContext);
diff --git a/Engine/Plugins/VirtualProduction/SceneState/Source/SceneState/Private/SceneStateObject.cpp b/Engine/Plugins/VirtualProduction/SceneState/Source/SceneState/Private/SceneStateObject.cpp
--- a/Engine/Plugins/VirtualProduction/SceneState/Source/SceneState/Private/SceneStateObject.cpp
+++ b/Engine/Plugins/VirtualProduction/SceneState/Source/SceneState/Private/SceneStateObject.cpp
@@ -1,6 +1,7 @@
 // Copyright Epic Games, Inc. All Rights Reserved.
 
 #include "SceneStateObject.h"
+#include "Debug/SceneStateDebugger.h"
 #include "PropertyBindingDataView.h"
 #include "SceneStateEvent.h"
 #include "SceneStateEventStream.h"
@@ -101,6 +102,10 @@ void USceneStateObject::Exit()
 
 	TemplateData = nullptr;
 	RootExecutionContext.Reset();
+
+#if WITH_EDITOR
+	SetDebugger(nullptr);
+#endif
 }
 
 TSharedRef<UE::SceneState::FExecutionContextRegistry> USceneStateObject::GetContextRegistry() const
@@ -108,6 +113,27 @@ TSharedRef<UE::SceneState::FExecutionContextRegistry> USceneStateObject::GetCont
 	return ContextRegistry.ToSharedRef();
 }
 
+#if WITH_EDITOR
+TSharedPtr<UE::SceneState::FDebugger> USceneStateObject::GetDebugger() const
+{
+	return DebuggerWeak.Pin();
+}
+
+void USceneStateObject::SetDebugger(const TWeakPtr<UE::SceneState::FDebugger>& InDebuggerWeak)
+{
+	if (InDebuggerWeak == DebuggerWeak)
+	{
+		return;
+	}
+
+	if (TSharedPtr<UE::SceneState::FDebugger> Debugger = GetDebugger())
+	{
+		Debugger->Detach();
+	}
+	DebuggerWeak = InDebuggerWeak;
+}
+#endif
+
 UWorld* USceneStateObject::GetWorld() const
 {
 	if (UObject* Context = GetContextObject())
diff --git a/Engine/Plugins/VirtualProduction/SceneState/Source/SceneState/Private/SceneStateTemplateData.cpp b/Engine/Plugins/VirtualProduction/SceneState/Source/SceneState/Private/SceneStateTemplateData.cpp
--- a/Engine/Plugins/VirtualProduction/SceneState/Source/SceneState/Private/SceneStateTemplateData.cpp
+++ b/Engine/Plugins/VirtualProduction/SceneState/Source/SceneState/Private/SceneStateTemplateData.cpp
@@ -168,44 +168,56 @@ const FSceneStateTask* USceneStateTemplateData::FindTaskFromNode(FObjectKey InTa
 	return nullptr;
 }
 
-void USceneStateTemplateData::ForEachStateInstance(const TSharedRef<UE::SceneState::FExecutionContextRegistry>& InContextRegistry
-	, FObjectKey InStateNode
-	, TFunctionRef<void(const FSceneStateInstance&)> InFunctor) const
+void USceneStateTemplateData::ForEachStateInstance(const TSharedRef<UE::SceneState::FExecutionContextRegistry>& InContextRegistry, uint16 InStateIndex, TFunctionRef<void(const FSceneStateInstance&)> InFunctor) const
 {
-	const uint16* StateIndex = StateNodeToIndex.Find(InStateNode);
-	if (!StateIndex)
+	if (!States.IsValidIndex(InStateIndex))
 	{
 		return;
 	}
 
 	InContextRegistry->ForEachExecutionContext(
-		[StateIndex, &InFunctor](const FSceneStateExecutionContext& InExecutionContext)
+		[InStateIndex, &InFunctor](const FSceneStateExecutionContext& InExecutionContext)
 		{
-			if (const FSceneStateInstance* StateInstance = InExecutionContext.FindStateInstance(*StateIndex))
+			if (const FSceneStateInstance* StateInstance = InExecutionContext.FindStateInstance(InStateIndex))
 			{
 				InFunctor(*StateInstance);
 			}
 		});
 }
 
-void USceneStateTemplateData::ForEachTaskInstance(const TSharedRef<UE::SceneState::FExecutionContextRegistry>& InContextRegistry, FObjectKey InTaskNode, TFunctionRef<void(const FSceneStateTaskInstance&)> InFunctor) const
+void USceneStateTemplateData::ForEachStateInstance(const TSharedRef<UE::SceneState::FExecutionContextRegistry>& InContextRegistry, FObjectKey InStateNode, TFunctionRef<void(const FSceneStateInstance&)> InFunctor) const
 {
-	const uint16* TaskIndex = TaskNodeToIndex.Find(InTaskNode);
-	if (!TaskIndex)
+	if (const uint16* StateIndex = StateNodeToIndex.Find(InStateNode))
+	{
+		ForEachStateInstance(InContextRegistry, *StateIndex, InFunctor);
+	}
+}
+
+void USceneStateTemplateData::ForEachTaskInstance(const TSharedRef<UE::SceneState::FExecutionContextRegistry>& InContextRegistry, uint16 InTaskIndex, TFunctionRef<void(const FSceneStateTaskInstance&)> InFunctor) const
+{
+	if (!Tasks.IsValidIndex(InTaskIndex))
 	{
 		return;
 	}
 
 	InContextRegistry->ForEachExecutionContext(
-		[TaskIndex, &InFunctor](const FSceneStateExecutionContext& InExecutionContext)
+		[InTaskIndex, &InFunctor](const FSceneStateExecutionContext& InExecutionContext)
 		{
-			if (const FSceneStateTaskInstance* TaskInstance = InExecutionContext.FindTaskInstance(*TaskIndex).GetPtr<const FSceneStateTaskInstance>())
+			if (const FSceneStateTaskInstance* TaskInstance = InExecutionContext.FindTaskInstance(InTaskIndex).GetPtr<const FSceneStateTaskInstance>())
 			{
 				InFunctor(*TaskInstance);
 			}
 		});
 }
 
+void USceneStateTemplateData::ForEachTaskInstance(const TSharedRef<UE::SceneState::FExecutionContextRegistry>& InContextRegistry, FObjectKey InTaskNode, TFunctionRef<void(const FSceneStateTaskInstance&)> InFunctor) const
+{
+	if (const uint16* TaskIndex = TaskNodeToIndex.Find(InTaskNode))
+	{
+		ForEachTaskInstance(InContextRegistry, *TaskIndex, InFunctor);
+	}
+}
+
 void USceneStateTemplateData::OnObjectsReinstanced(const UClass* InOwnerClass, const TMap<UObject*, UObject*>& InReplacementMap)
 {
 	TSet<const UStruct*> Structs;
diff --git a/Engine/Plugins/VirtualProduction/SceneState/Source/SceneState/Private/Tasks/SceneStateTaskExecutionContext.cpp b/Engine/Plugins/VirtualProduction/SceneState/Source/SceneState/Private/Tasks/SceneStateTaskExecutionContext.cpp
--- a/Engine/Plugins/VirtualProduction/SceneState/Source/SceneState/Private/Tasks/SceneStateTaskExecutionContext.cpp
+++ b/Engine/Plugins/VirtualProduction/SceneState/Source/SceneState/Private/Tasks/SceneStateTaskExecutionContext.cpp
@@ -13,15 +13,6 @@ namespace UE::SceneState
 namespace Private
 {
 
-FExecutionContextHandle GetContextHandle(const FSceneStateExecutionContext& InContext)
-{
-	if (TSharedPtr<const FExecutionContextRegistry> ContextRegistry = InContext.GetContextRegistry())
-	{
-		return ContextRegistry->FindHandle(InContext);
-	}
-	return FExecutionContextHandle();
-}
-
 uint16 GetStateInstanceId(const FSceneStateTask& InTask, const FSceneStateExecutionContext& InContext)
 {
 	if (FSceneStateInstance* StateInstance = InContext.FindStateInstance(InTask.GetParentStateIndex()))
@@ -50,7 +41,7 @@ FTaskExecutionContext::FTaskExecutionContext()
 }
 
 FTaskExecutionContext::FTaskExecutionContext(const FSceneStateTask& InTask, const FSceneStateExecutionContext& InContext)
-	: ContextHandle(Private::GetContextHandle(InContext))
+	: ContextHandle(InContext.GetContextHandle())
 	, ContextRegistryWeak(InContext.GetContextRegistry())
 	, TaskIndex(InTask.GetTaskIndex())
 	, StateIndex(InTask.GetParentStateIndex())
diff --git a/Engine/Plugins/VirtualProduction/SceneState/Source/SceneState/Public/SceneStateExecutionContext.h b/Engine/Plugins/VirtualProduction/SceneState/Source/SceneState/Public/SceneStateExecutionContext.h
--- a/Engine/Plugins/VirtualProduction/SceneState/Source/SceneState/Public/SceneStateExecutionContext.h
+++ b/Engine/Plugins/VirtualProduction/SceneState/Source/SceneState/Public/SceneStateExecutionContext.h
@@ -67,7 +67,10 @@ struct FSceneStateExecutionContext
 	UE_API ~FSceneStateExecutionContext();
 
 	/** Gets the registry that this context is registered to */
-	TSharedPtr<const UE::SceneState::FExecutionContextRegistry> GetContextRegistry() const;
+	TWeakPtr<const UE::SceneState::FExecutionContextRegistry> GetContextRegistry() const;
+
+	/** Gets the context handle for this execution context, if any */
+	UE::SceneState::FExecutionContextHandle GetContextHandle() const;
 
 	const USceneStateTemplateData* GetTemplateData() const
 	{
@@ -81,7 +84,7 @@ struct FSceneStateExecutionContext
 	UE_API const FString& GetExecutionContextName() const;
 
 	/** Retrieves the Root Object of this Execution */
-	UE_API UObject* GetRootObject() const;
+	UE_API USceneStateObject* GetRootObject() const;
 
 	/** Retrieves the Context Object for this Execution */
 	UE_API UObject* GetContextObject() const;
@@ -107,7 +110,7 @@ struct FSceneStateExecutionContext
 		/** Template data object holding all the state machines, states, etc. */
 		const USceneStateTemplateData* TemplateData = nullptr;
 		/** Root object instance */
-		UObject* RootObject = nullptr;
+		USceneStateObject* RootObject = nullptr;
 		/** Context object for utility purposes (e.g. to get world from context, etc) */
 		UObject* ContextObject = nullptr;
 		/** Event stream to use for execution */
@@ -202,6 +205,9 @@ struct FSceneStateExecutionContext
 	/** Removes the State Instance of this context for the given state */
 	void RemoveStateInstance(const FSceneState& InState) const;
 
+	/** Iterates the active state instances of this context */
+	void ForEachStateInstance(TFunctionRef<void(uint16, const FSceneStateInstance&)> InFunc) const;
+
 	/** Returns the Task of this context for the given task absolute index. Returns an invalid view if not found */
 	FConstStructView FindTask(uint16 InAbsoluteIndex) const;
 
@@ -268,7 +274,6 @@ struct FSceneStateExecutionContext
 	/** Called to set the status of the given state */
 	void SetStateStatus(const FSceneState& InState, FSceneStateInstance& InStateInstance, UE::SceneState::EExecutionStatus InStatus) const;
 
-private:
 	/** Retrieves the parent execution context if existing. Can be null */
 	const FSceneStateExecutionContext* FindParentContext() const;
 
@@ -284,6 +289,7 @@ private:
 	/** Retrieves the Absolute Index of a Conduit in the Conduits array */
 	uint16 GetConduitIndex(const FSceneStateConduit& InConduit) const;
 
+private:
 	/** Weak ptr to the context holding this execution */
 	UPROPERTY()
 	TWeakObjectPtr<UObject> ContextObjectWeak;
@@ -292,9 +298,9 @@ private:
 	UPROPERTY()
 	TWeakObjectPtr<USceneStateEventStream> EventStreamWeak;
 
-	/** Weak ptr to the root object owning this execution */
+	/** Pointer to the root object owning this execution */
 	UPROPERTY()
-	TObjectPtr<UObject> RootObject;
+	TObjectPtr<USceneStateObject> RootObject;
 
 	/** Description of the context */
 	UPROPERTY()
@@ -355,6 +361,9 @@ private:
 	 */
 	mutable TMap<uint16, FStructView> FunctionInstanceViews;
 
+	/** Context handle to this execution. */
+	UE::SceneState::FExecutionContextHandle ContextHandle;
+
 	/** Context handle to the parent execution context. Can be invalid. */
 	UE::SceneState::FExecutionContextHandle ParentHandle;
 
diff --git a/Engine/Plugins/VirtualProduction/SceneState/Source/SceneState/Public/SceneStateInstance.h b/Engine/Plugins/VirtualProduction/SceneState/Source/SceneState/Public/SceneStateInstance.h
--- a/Engine/Plugins/VirtualProduction/SceneState/Source/SceneState/Public/SceneStateInstance.h
+++ b/Engine/Plugins/VirtualProduction/SceneState/Source/SceneState/Public/SceneStateInstance.h
@@ -21,7 +21,7 @@ PRAGMA_DISABLE_DEPRECATION_WARNINGS
 	FSceneStateInstance& operator=(FSceneStateInstance&&) = default;
 PRAGMA_ENABLE_DEPRECATION_WARNINGS
 
-	/** Time elapsed since the state entered */
+	UE_DEPRECATED(5.8, "ElapsedTime is deprecated, it is no longer in use.")
 	float ElapsedTime = 0.f;
 
 	/** Gets the execution status of the state instance */
@@ -35,7 +35,7 @@ PRAGMA_ENABLE_DEPRECATION_WARNINGS
 	UE::SceneState::EExecutionStatus Status = UE::SceneState::EExecutionStatus::NotStarted;
 
 	/** Get an identifier for this instance */
-	uint16 GetInstanceId() const;
+	UE_INTERNAL UE_API uint16 GetInstanceId() const;
 
 private:
 	/** The id for this instance. This is used to differentiate state instances for a same state and consequently task instances for a same task*/
diff --git a/Engine/Plugins/VirtualProduction/SceneState/Source/SceneState/Public/SceneStateObject.h b/Engine/Plugins/VirtualProduction/SceneState/Source/SceneState/Public/SceneStateObject.h
--- a/Engine/Plugins/VirtualProduction/SceneState/Source/SceneState/Public/SceneStateObject.h
+++ b/Engine/Plugins/VirtualProduction/SceneState/Source/SceneState/Public/SceneStateObject.h
@@ -22,6 +22,7 @@ struct FSceneStateExecutionContext;
 
 namespace UE::SceneState
 {
+	class FDebugger;
 	class FExecutionContextRegistry;
 }
 
@@ -84,6 +85,14 @@ public:
 	/** Gets the context registry owned by this object */
 	UE_API TSharedRef<UE::SceneState::FExecutionContextRegistry> GetContextRegistry() const;
 
+#if WITH_EDITOR
+	/** Gets the currently attached debugger, if any */
+	UE_INTERNAL UE_API TSharedPtr<UE::SceneState::FDebugger> GetDebugger() const;
+
+	/** Attaches a new debugger, detaching the old one if any */
+	void SetDebugger(const TWeakPtr<UE::SceneState::FDebugger>& InDebuggerWeak);
+#endif
+
 	//~ Begin UObject
 	UE_API virtual UWorld* GetWorld() const override final;
 	UE_API virtual void BeginDestroy() override;
@@ -107,6 +116,11 @@ private:
 
 	/** The registry containing all the execution contexts in this object */
 	TSharedPtr<UE::SceneState::FExecutionContextRegistry> ContextRegistry;
+
+#if WITH_EDITOR
+	/** Debugger attached to this object */
+	TWeakPtr<UE::SceneState::FDebugger> DebuggerWeak;
+#endif
 };
 
 #undef UE_API
diff --git a/Engine/Plugins/VirtualProduction/SceneState/Source/SceneState/Public/SceneStateTemplateData.h b/Engine/Plugins/VirtualProduction/SceneState/Source/SceneState/Public/SceneStateTemplateData.h
--- a/Engine/Plugins/VirtualProduction/SceneState/Source/SceneState/Public/SceneStateTemplateData.h
+++ b/Engine/Plugins/VirtualProduction/SceneState/Source/SceneState/Public/SceneStateTemplateData.h
@@ -73,21 +73,31 @@ public:
 	UE_API const FSceneStateTask* FindTaskFromNode(FObjectKey InTaskNode) const;
 
 	/**
-	 * For a given root state and state node retrieves the mapped state instances
+	 * For a given root state and state index retrieves the mapped state instances
 	 * @param InContextRegistry the registry containing all the execution contexts
-	 * @param InStateNode the state node to look for 
+	 * @param InStateIndex the state index to look for
 	 * @param InFunctor the functor to run for each state instance
 	 */
+	UE_API void ForEachStateInstance(const TSharedRef<UE::SceneState::FExecutionContextRegistry>& InContextRegistry
+		, uint16 InStateIndex
+		, TFunctionRef<void(const FSceneStateInstance&)> InFunctor) const;
+
+	UE_DEPRECATED(5.8, "ForEachStateInstance deprecated. No longer supported. Use the overload that takes in a state index instead")
 	UE_API void ForEachStateInstance(const TSharedRef<UE::SceneState::FExecutionContextRegistry>& InContextRegistry
 		, FObjectKey InStateNode
 		, TFunctionRef<void(const FSceneStateInstance&)> InFunctor) const;
 
 	/**
-	 * For a given root state and task node retrieves the mapped task instances
+	 * For a given root state and task index retrieves the mapped task instances
 	 * @param InContextRegistry the registry containing all the execution contexts
-	 * @param InTaskNode the task node to look for
+	 * @param InTaskIndex the task index to look for
 	 * @param InFunctor the functor to run for each task instance
 	 */
+	UE_API void ForEachTaskInstance(const TSharedRef<UE::SceneState::FExecutionContextRegistry>& InContextRegistry
+		, uint16 InTaskIndex
+		, TFunctionRef<void(const FSceneStateTaskInstance&)> InFunctor) const;
+
+	UE_DEPRECATED(5.8, "ForEachTaskInstance deprecated. No longer supported. Use the overload that takes in a task index instead")
 	UE_API void ForEachTaskInstance(const TSharedRef<UE::SceneState::FExecutionContextRegistry>& InContextRegistry
 		, FObjectKey InTaskNode
 		, TFunctionRef<void(const FSceneStateTaskInstance&)> InFunctor) const;
@@ -176,13 +186,22 @@ public:
 	UPROPERTY()
 	TArray<FSceneStateTransitionMetadata> TransitionMetadata;
 
-	/** Map of a State Node to its State Index */
+	/** 
+	 * Map of a State Node to its State Index.
+	 * Set after a compilation occurs.
+	 */
 	TMap<FObjectKey, uint16> StateNodeToIndex;
 
-	/** Map of a State Machine Graph to its State Machine Index */
+	/** 
+	 * Map of a State Machine Graph to its State Machine Index.
+	 * Set after a compilation occurs.
+	 */
 	TMap<FObjectKey, uint16> StateMachineGraphToIndex;
 
-	/** Map of a Task Node to its Task Index */
+	/** 
+	 * Map of a Task Node to its Task Index.
+	 * Set after a compilation occurs.
+	 */
 	TMap<FObjectKey, uint16> TaskNodeToIndex;
 #endif
 };
diff --git a/Engine/Plugins/VirtualProduction/SceneState/Source/SceneStateBlueprint/Private/SceneStateBlueprint.cpp b/Engine/Plugins/VirtualProduction/SceneState/Source/SceneStateBlueprint/Private/SceneStateBlueprint.cpp
--- a/Engine/Plugins/VirtualProduction/SceneState/Source/SceneStateBlueprint/Private/SceneStateBlueprint.cpp
+++ b/Engine/Plugins/VirtualProduction/SceneState/Source/SceneStateBlueprint/Private/SceneStateBlueprint.cpp
@@ -270,7 +270,7 @@ void USceneStateBlueprint::SetObjectBeingDebugged(UObject* InNewObject)
 
 	Graph::FBlueprintDebugObjectChange Change;
 	Change.Blueprint = this;
-	Change.DebugObject = InNewObject;
+	Change.DebugObject = CurrentObjectBeingDebugged.Get();
 
 	Graph::OnBlueprintDebugObjectChanged.Broadcast(Change);
 }
diff --git a/Engine/Plugins/VirtualProduction/SceneState/Source/SceneStateBlueprint/Public/SceneStateBlueprintDelegates.h b/Engine/Plugins/VirtualProduction/SceneState/Source/SceneStateBlueprint/Public/SceneStateBlueprintDelegates.h
--- a/Engine/Plugins/VirtualProduction/SceneState/Source/SceneStateBlueprint/Public/SceneStateBlueprintDelegates.h
+++ b/Engine/Plugins/VirtualProduction/SceneState/Source/SceneStateBlueprint/Public/SceneStateBlueprintDelegates.h
@@ -22,7 +22,8 @@ struct FBlueprintDebugObjectChange
  * Delegate called when the debug object of a blueprint has changed.
  * Blueprint Editor doesn't have a dedicated delegate for this and instead calls IBlueprintEditor::RefreshMyBlueprint, or IBlueprintEditor::RefreshEditors.
  * Additionally, these refresh functions are not called in every place the debug object is changed (e.g. in UWorld::TransferBlueprintDebugReferences).
- * This delegate solves this by broadcasting in USceneStateBlueprint::SetObjectBeingDebugged 
+ * UBlueprint::DebuggingWorldRegistrationHelper is closer to what is needed, but still requires a valid UWorld to have been found.
+ * This delegate solves this by broadcasting in USceneStateBlueprint::SetObjectBeingDebugged without requiring a valid world.
  */
 extern SCENESTATEBLUEPRINT_API TMulticastDelegate<void(const FBlueprintDebugObjectChange&)> OnBlueprintDebugObjectChanged;
 
diff --git a/Engine/Plugins/VirtualProduction/SceneState/Source/SceneStateBlueprintEditor/Private/Compiler/SceneStateBlueprintCompilerContext.cpp b/Engine/Plugins/VirtualProduction/SceneState/Source/SceneStateBlueprintEditor/Private/Compiler/SceneStateBlueprintCompilerContext.cpp
--- a/Engine/Plugins/VirtualProduction/SceneState/Source/SceneStateBlueprintEditor/Private/Compiler/SceneStateBlueprintCompilerContext.cpp
+++ b/Engine/Plugins/VirtualProduction/SceneState/Source/SceneStateBlueprintEditor/Private/Compiler/SceneStateBlueprintCompilerContext.cpp
@@ -2,6 +2,8 @@
 
 #include "SceneStateBlueprintCompilerContext.h"
 #include "Kismet2/KismetReinstanceUtilities.h"
+#include "Nodes/SceneStateMachineStateNode.h"
+#include "Nodes/SceneStateMachineTaskNode.h"
 #include "SceneStateBindingCompiler.h"
 #include "SceneStateBlueprint.h"
 #include "SceneStateGeneratedClass.h"
@@ -185,6 +187,7 @@ void FBlueprintCompilerContext::MergeUbergraphPagesIn(UEdGraph* InUbergraph)
 	NewGeneratedClass->TemplateData->StateMachineGraphToIndex.Append(MoveTemp(StateMachineGraphToIndex));
 
 	UpdateTaskIndices();
+	SetCompiledIndices();
 	CompileBindings();
 }
 
@@ -252,6 +255,29 @@ void FBlueprintCompilerContext::UpdateTaskIndices()
 	}
 }
 
+void FBlueprintCompilerContext::SetCompiledIndices()
+{
+	check(NewGeneratedClass->TemplateData);
+
+	for (TPair<FObjectKey, uint16>& Pair : NewGeneratedClass->TemplateData->StateNodeToIndex)
+	{
+		USceneStateMachineStateNode* const StateNode = CastChecked<USceneStateMachineStateNode>(Pair.Key.ResolveObjectPtr());
+		StateNode->SetCompiledIndex(Pair.Value);
+	}
+
+	for (TPair<FObjectKey, uint16>& Pair : NewGeneratedClass->TemplateData->StateMachineGraphToIndex)
+	{
+		USceneStateMachineGraph* const StateMachineGraph = CastChecked<USceneStateMachineGraph>(Pair.Key.ResolveObjectPtr());
+		StateMachineGraph->SetCompiledIndex(Pair.Value);
+	}
+
+	for (TPair<FObjectKey, uint16>& Pair : NewGeneratedClass->TemplateData->TaskNodeToIndex)
+	{
+		USceneStateMachineTaskNode* const TaskNode = CastChecked<USceneStateMachineTaskNode>(Pair.Key.ResolveObjectPtr());
+		TaskNode->SetCompiledIndex(Pair.Value);
+	}
+}
+
 void FBlueprintCompilerContext::CompileBindings()
 {
 	FBindingCompiler BindingCompiler(*this, CastChecked<USceneStateBlueprint>(Blueprint), NewGeneratedClass->TemplateData);
diff --git a/Engine/Plugins/VirtualProduction/SceneState/Source/SceneStateBlueprintEditor/Private/Compiler/SceneStateBlueprintCompilerContext.h b/Engine/Plugins/VirtualProduction/SceneState/Source/SceneStateBlueprintEditor/Private/Compiler/SceneStateBlueprintCompilerContext.h
--- a/Engine/Plugins/VirtualProduction/SceneState/Source/SceneStateBlueprintEditor/Private/Compiler/SceneStateBlueprintCompilerContext.h
+++ b/Engine/Plugins/VirtualProduction/SceneState/Source/SceneStateBlueprintEditor/Private/Compiler/SceneStateBlueprintCompilerContext.h
@@ -42,6 +42,9 @@ private:
 	/** Sets the Task Index and Parent State Index of each Task */
 	void UpdateTaskIndices();
 
+	/** Sets the indices of each mapped type (State, StateMachine, Task) to their editor node for debugging */
+	void SetCompiledIndices();
+
 	void CompileBindings();
 
 	USceneStateGeneratedClass* NewGeneratedClass = nullptr;
diff --git a/Engine/Plugins/VirtualProduction/SceneState/Source/SceneStateBlueprintEditor/Private/SceneStateBlueprintEditor.cpp b/Engine/Plugins/VirtualProduction/SceneState/Source/SceneStateBlueprintEditor/Private/SceneStateBlueprintEditor.cpp
--- a/Engine/Plugins/VirtualProduction/SceneState/Source/SceneStateBlueprintEditor/Private/SceneStateBlueprintEditor.cpp
+++ b/Engine/Plugins/VirtualProduction/SceneState/Source/SceneStateBlueprintEditor/Private/SceneStateBlueprintEditor.cpp
@@ -3,6 +3,7 @@
 #include "SceneStateBlueprintEditor.h"
 #include "AppModes/SceneStateAppModes.h"
 #include "AppModes/SceneStateBlueprintEditorMode.h"
+#include "Debug/SceneStateDebugger.h"
 #include "DetailsView/SceneStateBindingsChildrenCustomization.h"
 #include "DetailsView/SceneStateMachineTaskNodeCustomization.h"
 #include "DetailsView/SceneStatePropertyBindingExtension.h"
@@ -10,10 +11,12 @@
 #include "Kismet2/BlueprintEditorUtils.h"
 #include "Nodes/SceneStateMachineTaskNode.h"
 #include "SceneStateBlueprint.h"
+#include "SceneStateBlueprintDelegates.h"
 #include "SceneStateBlueprintEditorCommands.h"
 #include "SceneStateBlueprintExtension.h"
 #include "SceneStateBlueprintFactory.h"
 #include "SceneStateMachineGraph.h"
+#include "SceneStateObject.h"
 #include "SceneStateTransitionGraph.h"
 #include "ScopedTransaction.h"
 #include "Tools/BaseAssetToolkit.h"
@@ -24,7 +27,18 @@
 namespace UE::SceneState::Editor
 {
 
-const FName FSceneStateBlueprintEditor::SelectionState_StateMachine(TEXT("StateMachine"));
+const FLazyName FSceneStateBlueprintEditor::SelectionState_StateMachine(TEXT("StateMachine"));
+
+FSceneStateBlueprintEditor::FSceneStateBlueprintEditor()
+	: Debugger(MakeShared<FDebugger>())
+{
+}
+
+FSceneStateBlueprintEditor::~FSceneStateBlueprintEditor()
+{
+	Graph::OnBlueprintDebugObjectChanged.Remove(OnDebugObjectChangedHandle);
+	OnDebugObjectChangedHandle.Reset();
+}
 
 void FSceneStateBlueprintEditor::Init(USceneStateBlueprint* InBlueprint, const FAssetOpenArgs& InOpenArgs)
 {
@@ -32,6 +46,8 @@ void FSceneStateBlueprintEditor::Init(USceneStateBlueprint* InBlueprint, const F
 
 	AddBlueprintExtensions(*InBlueprint);
 
+	OnDebugObjectChangedHandle = Graph::OnBlueprintDebugObjectChanged.AddSP(this, &FSceneStateBlueprintEditor::OnDebugObjectChanged);
+
 	Super::InitBlueprintEditor(InOpenArgs.GetToolkitMode()
 		, InOpenArgs.ToolkitHost
 		, { InBlueprint }
@@ -65,6 +81,17 @@ void FSceneStateBlueprintEditor::AddBlueprintExtensions(USceneStateBlueprint& In
 	}
 }
 
+void FSceneStateBlueprintEditor::OnDebugObjectChanged(const Graph::FBlueprintDebugObjectChange& InDebugObjectChange)
+{
+	Debugger.Reset();
+
+	if (USceneStateObject* Object = Cast<USceneStateObject>(InDebugObjectChange.DebugObject))
+	{
+		Debugger = MakeShared<FDebugger>();
+		Debugger->AttachTo(Object);
+	}
+}
+
 void FSceneStateBlueprintEditor::RefreshEditors(ERefreshBlueprintEditorReason::Type InReason)
 {
 	Super::RefreshEditors(InReason);
diff --git a/Engine/Plugins/VirtualProduction/SceneState/Source/SceneStateBlueprintEditor/Private/SceneStateBlueprintEditor.h b/Engine/Plugins/VirtualProduction/SceneState/Source/SceneStateBlueprintEditor/Private/SceneStateBlueprintEditor.h
--- a/Engine/Plugins/VirtualProduction/SceneState/Source/SceneStateBlueprintEditor/Private/SceneStateBlueprintEditor.h
+++ b/Engine/Plugins/VirtualProduction/SceneState/Source/SceneStateBlueprintEditor/Private/SceneStateBlueprintEditor.h
@@ -7,9 +7,18 @@
 class USceneStateBlueprint;
 struct FAssetOpenArgs;
 
-namespace UE::SceneState::Editor
+namespace UE::SceneState
 {
-	class SStateMachineMenu;
+	class FDebugger;
+
+	namespace Editor
+	{
+		class SStateMachineMenu;
+	}
+	namespace Graph
+	{
+		struct FBlueprintDebugObjectChange;
+	}
 }
 
 namespace UE::SceneState::Editor
@@ -21,7 +30,10 @@ class FSceneStateBlueprintEditor : public FBlueprintEditor
 public:
 	using Super = FBlueprintEditor;
 
-	static const FName SelectionState_StateMachine;
+	static const FLazyName SelectionState_StateMachine;
+
+	FSceneStateBlueprintEditor();
+	virtual ~FSceneStateBlueprintEditor() override;
 
 	void Init(USceneStateBlueprint* InBlueprint, const FAssetOpenArgs& InOpenArgs);
 
@@ -31,6 +43,9 @@ public:
 
 	void AddBlueprintExtensions(USceneStateBlueprint& InBlueprint);
 
+	/** Called when the debug object has changed */
+	void OnDebugObjectChanged(const Graph::FBlueprintDebugObjectChange& InDebugObjectChange);
+
 	//~ Begin IBlueprintEditor
 	virtual void RefreshEditors(ERefreshBlueprintEditorReason::Type InReason) override;
 	virtual void RefreshMyBlueprint();
@@ -53,7 +68,13 @@ public:
 	//~ End IToolkit Interface
 
 private:
+	/** The current active debugger, or null if no debugger is attached */
+	TSharedPtr<UE::SceneState::FDebugger> Debugger;
+
 	TSharedPtr<SStateMachineMenu> StateMachineMenu;
+
+	/** Handle to the delegate called when the debug object has changed */
+	FDelegateHandle OnDebugObjectChangedHandle;
 };
 
 } // UE::SceneState::Editor
diff --git a/Engine/Plugins/VirtualProduction/SceneState/Source/SceneStateMachineEditor/Private/SceneStateMachineEditorStyle.cpp b/Engine/Plugins/VirtualProduction/SceneState/Source/SceneStateMachineEditor/Private/SceneStateMachineEditorStyle.cpp
--- a/Engine/Plugins/VirtualProduction/SceneState/Source/SceneStateMachineEditor/Private/SceneStateMachineEditorStyle.cpp
+++ b/Engine/Plugins/VirtualProduction/SceneState/Source/SceneStateMachineEditor/Private/SceneStateMachineEditorStyle.cpp
@@ -27,6 +27,7 @@ FStateMachineEditorStyle::FStateMachineEditorStyle()
 
 	Set("SpillColor.State.Inactive", FLinearColor(0.08f, 0.08f, 0.08f));
 	Set("SpillColor.State.Active", FLinearColor(1.f, 0.6f, 0.35f));
+	Set("SpillColor.State.Finished", FLinearColor(0.25, 1.f, 0.25f));
 	Set("SpillColor.Conduit", FLinearColor(0.38f, 0.45f, 0.21f));
 
 	Set("SpillColor.Task.Inactive", FLinearColor(0.08f, 0.08f, 0.08f));
diff --git a/Engine/Plugins/VirtualProduction/SceneState/Source/SceneStateMachineEditor/Private/Widgets/SSceneStateMachineStateNode.cpp b/Engine/Plugins/VirtualProduction/SceneState/Source/SceneStateMachineEditor/Private/Widgets/SSceneStateMachineStateNode.cpp
--- a/Engine/Plugins/VirtualProduction/SceneState/Source/SceneStateMachineEditor/Private/Widgets/SSceneStateMachineStateNode.cpp
+++ b/Engine/Plugins/VirtualProduction/SceneState/Source/SceneStateMachineEditor/Private/Widgets/SSceneStateMachineStateNode.cpp
@@ -1,6 +1,8 @@
 // Copyright Epic Games, Inc. All Rights Reserved.
 
 #include "SSceneStateMachineStateNode.h"
+#include "Debug/SceneStateDebugStateContext.h"
+#include "Debug/SceneStateDebugger.h"
 #include "Kismet2/BlueprintEditorUtils.h"
 #include "Nodes/SceneStateMachineStateNode.h"
 #include "SGraphPanel.h"
@@ -30,6 +32,7 @@ void SStateMachineStateNode::Construct(const FArguments& InArgs, USceneStateMach
 	const FStateMachineEditorStyle& Style = FStateMachineEditorStyle::Get();
 
 	ActiveColor = Style.GetColor("SpillColor.State.Active");
+	FinishedColor = Style.GetColor("SpillColor.State.Finished");
 	InactiveColor = Style.GetColor("SpillColor.State.Inactive");
 
 	GraphNode = InNode;
@@ -87,7 +90,25 @@ TSharedRef<SWidget> SStateMachineStateNode::MakeNodeInnerWidget()
 
 FSlateColor SStateMachineStateNode::GetStateBackgroundColor() const
 {
-	return bIsActiveState ? ActiveColor : InactiveColor;
+	if (!StateStatus.IsSet())
+	{
+		return InactiveColor;
+	}
+
+	const FLinearColor* TargetColor = &InactiveColor;
+
+	switch (*StateStatus)
+	{
+	case EExecutionStatus::Running:
+		TargetColor = &ActiveColor;
+		break;
+
+	case EExecutionStatus::Finished:
+		TargetColor = &FinishedColor;
+		break;
+	}
+
+	return FMath::Lerp(InactiveColor, *TargetColor, ActivePercentage);
 }
 
 void SStateMachineStateNode::UpdateGraphNode()
@@ -176,37 +197,61 @@ TSharedPtr<SToolTip> SStateMachineStateNode::GetComplexTooltip()
 
 void SStateMachineStateNode::GetNodeInfoPopups(FNodeInfoContext* InContext, TArray<FGraphInformationPopupInfo>& OutPopups) const
 {
-	bIsActiveState = false;
+	ActivePercentage = 0.f;
+	StateStatus.Reset();
 
-	UBlueprint* Blueprint = FBlueprintEditorUtils::FindBlueprintForNode(GraphNode);
+	const USceneStateMachineStateNode* const StateNode = Cast<USceneStateMachineStateNode>(GraphNode);
+	if (!StateNode || !StateNode->GetCompiledIndex().IsSet())
+	{
+		return;
+	}
+
+	const UBlueprint* const Blueprint = FBlueprintEditorUtils::FindBlueprintForNode(GraphNode);
 	if (!Blueprint)
 	{
 		return;
 	}
 
-	USceneStateObject* RootState = Cast<USceneStateObject>(Blueprint->GetObjectBeingDebugged());
+	const USceneStateObject* const RootState = Cast<USceneStateObject>(Blueprint->GetObjectBeingDebugged());
 	if (!RootState)
 	{
 		return;
 	}
 
-	const USceneStateTemplateData* TemplateData = RootState->GetTemplateData();
-	if (!TemplateData)
+	const TSharedPtr<UE::SceneState::FDebugger> Debugger = RootState->GetDebugger();
+	if (!Debugger.IsValid())
 	{
 		return;
 	}
 
-	bIsActiveState = false;
-
-	TemplateData->ForEachStateInstance(RootState->GetContextRegistry(), GraphNode,
-		[&OutPopups, &bIsActiveState = bIsActiveState, &ActiveColor = ActiveColor](const FSceneStateInstance& InInstance)
+	Debugger->ForEachDebugInstanceOfState(*StateNode->GetCompiledIndex(),
+		[&OutPopups, This=this](const FDebugStateInstance& InDebugInstance)
 		{
-			bIsActiveState = true;
-
-			FText StateText = FText::Format(LOCTEXT("StateStatusFormat", "Active for {0} s")
-				, FText::AsNumber(InInstance.ElapsedTime, &FStateMachineEditorStyle::Get().GetDefaultNumberFormat()));
-
-			OutPopups.Emplace(nullptr, ActiveColor, StateText.ToString());
+			const float CurrentActivePercentage = InDebugInstance.GetActivePercentage();
+			This->ActivePercentage = FMath::Max(This->ActivePercentage, CurrentActivePercentage);
+
+			switch (InDebugInstance.GetExecutionStatus())
+			{
+			case EExecutionStatus::Running:
+				{
+					This->StateStatus = EExecutionStatus::Running;
+
+					const FText StateText = FText::Format(LOCTEXT("StateStatusFormat", "Active for {0} s")
+						, FText::AsNumber(InDebugInstance.GetStatusElapsedTime(), &FStateMachineEditorStyle::Get().GetDefaultNumberFormat()));
+
+					OutPopups.Emplace(/*Icon*/nullptr
+						, FMath::Lerp(This->InactiveColor, This->ActiveColor, CurrentActivePercentage)
+						, StateText.ToString());
+				}
+				break;
+
+			case EExecutionStatus::Finished:
+				if (!This->StateStatus.IsSet())
+				{
+					This->StateStatus = EExecutionStatus::Finished;
+				}
+				break;
+			}
 		});
 }
 
diff --git a/Engine/Plugins/VirtualProduction/SceneState/Source/SceneStateMachineEditor/Private/Widgets/SSceneStateMachineStateNode.h b/Engine/Plugins/VirtualProduction/SceneState/Source/SceneStateMachineEditor/Private/Widgets/SSceneStateMachineStateNode.h
--- a/Engine/Plugins/VirtualProduction/SceneState/Source/SceneStateMachineEditor/Private/Widgets/SSceneStateMachineStateNode.h
+++ b/Engine/Plugins/VirtualProduction/SceneState/Source/SceneStateMachineEditor/Private/Widgets/SSceneStateMachineStateNode.h
@@ -2,6 +2,7 @@
 
 #pragma once
 
+#include "SceneStateEnums.h"
 #include "SSceneStateMachineNode.h"
 
 class USceneStateMachineStateNode;
@@ -42,9 +43,18 @@ private:
 	TSharedPtr<SOverlay> OutputPinOverlay;
 
 	FLinearColor ActiveColor = FLinearColor(ForceInit);
+	FLinearColor FinishedColor = FLinearColor(ForceInit);
 	FLinearColor InactiveColor = FLinearColor(ForceInit);
 
-	mutable bool bIsActiveState = false;
+	/**
+	 * Represents how 'active' the state is 
+	 * 1 means fully active, 0 means fully inactive.
+	 * When multiple instances are active, it will use the highest value.
+	 */
+	mutable float ActivePercentage = 0.f;
+
+	/** Whether the state should show the finish color instead of the active color */
+	mutable TOptional<EExecutionStatus> StateStatus;
 };
 
 } // UE::SceneState::Editor
diff --git a/Engine/Plugins/VirtualProduction/SceneState/Source/SceneStateMachineEditor/Private/Widgets/SSceneStateMachineTaskNode.cpp b/Engine/Plugins/VirtualProduction/SceneState/Source/SceneStateMachineEditor/Private/Widgets/SSceneStateMachineTaskNode.cpp
--- a/Engine/Plugins/VirtualProduction/SceneState/Source/SceneStateMachineEditor/Private/Widgets/SSceneStateMachineTaskNode.cpp
+++ b/Engine/Plugins/VirtualProduction/SceneState/Source/SceneStateMachineEditor/Private/Widgets/SSceneStateMachineTaskNode.cpp
@@ -178,13 +178,19 @@ void SStateMachineTaskNode::GetNodeInfoPopups(FNodeInfoContext* InContext, TArra
 {
 	TaskStatus = EExecutionStatus::NotStarted;
 
-	UBlueprint* Blueprint = FBlueprintEditorUtils::FindBlueprintForNode(GraphNode);
+	const USceneStateMachineTaskNode* const Node = Cast<USceneStateMachineTaskNode>(GraphNode);
+	if (!Node || !Node->GetCompiledIndex().IsSet())
+	{
+		return;
+	}
+
+	const UBlueprint* const Blueprint = FBlueprintEditorUtils::FindBlueprintForNode(GraphNode);
 	if (!Blueprint)
 	{
 		return;
 	}
 
-	USceneStateObject* RootState = Cast<USceneStateObject>(Blueprint->GetObjectBeingDebugged());
+	const USceneStateObject* const RootState = Cast<USceneStateObject>(Blueprint->GetObjectBeingDebugged());
 	if (!RootState)
 	{
 		return;
@@ -198,7 +204,7 @@ void SStateMachineTaskNode::GetNodeInfoPopups(FNodeInfoContext* InContext, TArra
 
 	TOptional<EExecutionStatus> ResultStatus;
 
-	TemplateData->ForEachTaskInstance(RootState->GetContextRegistry(), GraphNode,
+	TemplateData->ForEachTaskInstance(RootState->GetContextRegistry(), *Node->GetCompiledIndex(),
 		[&OutPopups, &ResultStatus, &StatusColors = StatusColors](const FSceneStateTaskInstance& InTaskInstance)
 		{
 			// Get the lowest status to show
diff --git a/Engine/Plugins/VirtualProduction/SceneState/Source/SceneStateMachineGraph/Private/Nodes/SceneStateMachineStateNode.cpp b/Engine/Plugins/VirtualProduction/SceneState/Source/SceneStateMachineGraph/Private/Nodes/SceneStateMachineStateNode.cpp
--- a/Engine/Plugins/VirtualProduction/SceneState/Source/SceneStateMachineGraph/Private/Nodes/SceneStateMachineStateNode.cpp
+++ b/Engine/Plugins/VirtualProduction/SceneState/Source/SceneStateMachineGraph/Private/Nodes/SceneStateMachineStateNode.cpp
@@ -15,6 +15,16 @@ USceneStateMachineStateNode::USceneStateMachineStateNode()
 	bCanRenameNode = true;
 }
 
+TOptional<uint16> USceneStateMachineStateNode::GetCompiledIndex() const
+{
+	return CompiledIndex;
+}
+
+void USceneStateMachineStateNode::SetCompiledIndex(uint16 InCompiledIndex)
+{
+	CompiledIndex = InCompiledIndex;
+}
+
 bool USceneStateMachineStateNode::FindEventHandlerId(const FSceneStateEventSchemaHandle& InEventSchemaHandle, FGuid& OutHandlerId) const
 {
 	for (const FSceneStateEventHandler& EventHandler : EventHandlers)
diff --git a/Engine/Plugins/VirtualProduction/SceneState/Source/SceneStateMachineGraph/Private/Nodes/SceneStateMachineTaskNode.cpp b/Engine/Plugins/VirtualProduction/SceneState/Source/SceneStateMachineGraph/Private/Nodes/SceneStateMachineTaskNode.cpp
--- a/Engine/Plugins/VirtualProduction/SceneState/Source/SceneStateMachineGraph/Private/Nodes/SceneStateMachineTaskNode.cpp
+++ b/Engine/Plugins/VirtualProduction/SceneState/Source/SceneStateMachineGraph/Private/Nodes/SceneStateMachineTaskNode.cpp
@@ -84,6 +84,16 @@ bool USceneStateMachineTaskNode::FindDataViewById(const FGuid& InStructId, FStru
 	return false;
 }
 
+TOptional<uint16> USceneStateMachineTaskNode::GetCompiledIndex() const
+{
+	return CompiledIndex;
+}
+
+void USceneStateMachineTaskNode::SetCompiledIndex(uint16 InCompiledIndex)
+{
+	CompiledIndex = InCompiledIndex;
+}
+
 void USceneStateMachineTaskNode::AllocateDefaultPins()
 {
 	CreatePin(EGPD_Input, USceneStateMachineGraphSchema::PC_Task, USceneStateMachineGraphSchema::PN_In);
diff --git a/Engine/Plugins/VirtualProduction/SceneState/Source/SceneStateMachineGraph/Private/SceneStateMachineGraph.cpp b/Engine/Plugins/VirtualProduction/SceneState/Source/SceneStateMachineGraph/Private/SceneStateMachineGraph.cpp
--- a/Engine/Plugins/VirtualProduction/SceneState/Source/SceneStateMachineGraph/Private/SceneStateMachineGraph.cpp
+++ b/Engine/Plugins/VirtualProduction/SceneState/Source/SceneStateMachineGraph/Private/SceneStateMachineGraph.cpp
@@ -59,6 +59,16 @@ bool USceneStateMachineGraph::IsRootStateMachine() const
 	return !!Cast<UBlueprint>(GetOuter());
 }
 
+TOptional<uint16> USceneStateMachineGraph::GetCompiledIndex() const
+{
+	return CompiledIndex;
+}
+
+void USceneStateMachineGraph::SetCompiledIndex(uint16 InCompiledIndex)
+{
+	CompiledIndex = InCompiledIndex;
+}
+
 void USceneStateMachineGraph::AddNode(UEdGraphNode* InNodeToAdd, bool bInUserAction, bool bInSelectNewNode)
 {
 	if (!InNodeToAdd)
diff --git a/Engine/Plugins/VirtualProduction/SceneState/Source/SceneStateMachineGraph/Public/Nodes/SceneStateMachineStateNode.h b/Engine/Plugins/VirtualProduction/SceneState/Source/SceneStateMachineGraph/Public/Nodes/SceneStateMachineStateNode.h
--- a/Engine/Plugins/VirtualProduction/SceneState/Source/SceneStateMachineGraph/Public/Nodes/SceneStateMachineStateNode.h
+++ b/Engine/Plugins/VirtualProduction/SceneState/Source/SceneStateMachineGraph/Public/Nodes/SceneStateMachineStateNode.h
@@ -7,6 +7,8 @@
 #include "SceneStateMachineNode.h"
 #include "SceneStateMachineStateNode.generated.h"
 
+#define UE_API SCENESTATEMACHINEGRAPH_API
+
 UCLASS(MinimalAPI, meta=(ToolTip="State node in a State Machine"))
 class USceneStateMachineStateNode : public USceneStateMachineNode, public ISceneStateEventHandlerProvider
 {
@@ -25,6 +27,12 @@ public:
 		return EventHandlers;
 	}
 
+	/** Gets the index to the state this node compiled to last */
+	UE_INTERNAL UE_API TOptional<uint16> GetCompiledIndex() const;
+
+	/** Sets the index to the state this node compiled to last */
+	UE_INTERNAL UE_API void SetCompiledIndex(uint16 InCompiledIndex);
+
 	//~ Begin ISceneStateEventHandlerProvider
 	virtual bool FindEventHandlerId(const FSceneStateEventSchemaHandle& InEventSchemaHandle, FGuid& OutHandlerId) const override;
 	//~ End ISceneStateEventHandlerProvider
@@ -51,4 +59,11 @@ public:
 
 	UPROPERTY(EditAnywhere, Category="Events", meta=(NoBinding))
 	TArray<FSceneStateEventHandler> EventHandlers;
+
+private:
+	/** Index to the state this state node compiled to last. */
+	UPROPERTY()
+	TOptional<uint16> CompiledIndex;
 };
+
+#undef UE_API
diff --git a/Engine/Plugins/VirtualProduction/SceneState/Source/SceneStateMachineGraph/Public/Nodes/SceneStateMachineTaskNode.h b/Engine/Plugins/VirtualProduction/SceneState/Source/SceneStateMachineGraph/Public/Nodes/SceneStateMachineTaskNode.h
--- a/Engine/Plugins/VirtualProduction/SceneState/Source/SceneStateMachineGraph/Public/Nodes/SceneStateMachineTaskNode.h
+++ b/Engine/Plugins/VirtualProduction/SceneState/Source/SceneStateMachineGraph/Public/Nodes/SceneStateMachineTaskNode.h
@@ -8,6 +8,8 @@
 #include "Templates/SubclassOf.h"
 #include "SceneStateMachineTaskNode.generated.h"
 
+#define UE_API SCENESTATEMACHINEGRAPH_API
+
 class USceneStateBlueprintableTask;
 struct FSceneStateTaskDescContext;
 struct FSceneStateTaskDescMutableContext;
@@ -41,17 +43,23 @@ public:
 		return OnPostEditTaskDelegate;
 	}
 
-	SCENESTATEMACHINEGRAPH_API bool SetTaskStruct(const UScriptStruct* InTaskStruct);
+	UE_API bool SetTaskStruct(const UScriptStruct* InTaskStruct);
 
-	SCENESTATEMACHINEGRAPH_API void SetTaskBlueprintClass(TSubclassOf<USceneStateBlueprintableTask> InTaskClass);
+	UE_API void SetTaskBlueprintClass(TSubclassOf<USceneStateBlueprintableTask> InTaskClass);
 
-	SCENESTATEMACHINEGRAPH_API bool FindDataViewById(const FGuid& InStructId, FStructView& OutDataView);
+	UE_API bool FindDataViewById(const FGuid& InStructId, FStructView& OutDataView);
 
 	const FGuid& GetTaskId() const
 	{
 		return TaskId;
 	}
 
+	/** Gets the index to the task this node compiled to last */
+	UE_INTERNAL UE_API TOptional<uint16> GetCompiledIndex() const;
+
+	/** Sets the index to the task this node compiled to last */
+	UE_INTERNAL UE_API void SetCompiledIndex(uint16 InCompiledIndex);
+
 	//~ Begin UEdGraphNode
 	virtual void AllocateDefaultPins() override;
 	virtual void OnRenameNode(const FString& InNodeName);
@@ -112,6 +120,10 @@ private:
 	UPROPERTY(EditAnywhere, Category="Task", meta=(StructTypeConst))
 	TInstancedStruct<FSceneStateTaskInstance> TaskInstance;
 
+	/** Index to the task this task node compiled to last. */
+	UPROPERTY()
+	TOptional<uint16> CompiledIndex;
+
 	UPROPERTY()
 	FGuid TaskId;
 
@@ -119,3 +131,5 @@ private:
 
 	mutable FText CachedTaskDisplayName;
 };
+
+#undef UE_API
diff --git a/Engine/Plugins/VirtualProduction/SceneState/Source/SceneStateMachineGraph/Public/SceneStateMachineGraph.h b/Engine/Plugins/VirtualProduction/SceneState/Source/SceneStateMachineGraph/Public/SceneStateMachineGraph.h
--- a/Engine/Plugins/VirtualProduction/SceneState/Source/SceneStateMachineGraph/Public/SceneStateMachineGraph.h
+++ b/Engine/Plugins/VirtualProduction/SceneState/Source/SceneStateMachineGraph/Public/SceneStateMachineGraph.h
@@ -34,6 +34,12 @@ public:
 	/** Returns whether the state machine is at the root level (i.e. directly under the blueprint and not nested under a state node) */
 	UE_API bool IsRootStateMachine() const;
 
+	/** Get the index to the state machine this node compiled to last */
+	UE_INTERNAL UE_API TOptional<uint16> GetCompiledIndex() const;
+
+	/** Sets the index to the state machine this node compiled to last */
+	UE_INTERNAL UE_API void SetCompiledIndex(uint16 InCompiledIndex);
+
 	//~ Begin UEdGraph
 	UE_API virtual void AddNode(UEdGraphNode* InNodeToAdd, bool bInUserAction, bool bInSelectNewNode) override;
 	//~ End UEdGraph
@@ -65,6 +71,10 @@ public:
 	ESceneStateMachineRunMode RunMode = ESceneStateMachineRunMode::Auto;
 
 private:
+	/** Index to the state machine this graph compiled to last. */
+	UPROPERTY()
+	TOptional<uint16> CompiledIndex;
+
 	static FOnParametersChanged OnParametersChangedDelegate;
 };
 
