Added missing include of scene state object

#Synced-CL 50160130

diff --git a/Engine/Plugins/VirtualProduction/SceneState/Source/SceneState/Private/Conduit/SceneStateConduit.cpp b/Engine/Plugins/VirtualProduction/SceneState/Source/SceneState/Private/Conduit/SceneStateConduit.cpp
--- a/Engine/Plugins/VirtualProduction/SceneState/Source/SceneState/Private/Conduit/SceneStateConduit.cpp
+++ b/Engine/Plugins/VirtualProduction/SceneState/Source/SceneState/Private/Conduit/SceneStateConduit.cpp
@@ -5,6 +5,7 @@
 #include "Conduit/SceneStateConduitLink.h"
 #include "SceneState.h"
 #include "SceneStateExecutionContext.h"
+#include "SceneStateObject.h"
 #include "Transition/SceneStateTransition.h"
 #include "Transition/SceneStateTransitionEvaluation.h"
 #include "Transition/SceneStateTransitionResult.h"
diff --git a/Engine/Plugins/VirtualProduction/SceneState/Source/SceneState/Private/Transition/SceneStateTransition.cpp b/Engine/Plugins/VirtualProduction/SceneState/Source/SceneState/Private/Transition/SceneStateTransition.cpp
--- a/Engine/Plugins/VirtualProduction/SceneState/Source/SceneState/Private/Transition/SceneStateTransition.cpp
+++ b/Engine/Plugins/VirtualProduction/SceneState/Source/SceneState/Private/Transition/SceneStateTransition.cpp
@@ -7,6 +7,7 @@
 #include "SceneStateEventStream.h"
 #include "SceneStateExecutionContext.h"
 #include "SceneStateMachine.h"
+#include "SceneStateObject.h"
 #include "SceneStateUtils.h"
 #include "Transition/SceneStateTransitionEvaluation.h"
 #include "Transition/SceneStateTransitionInstance.h"
