Motion Design Scene State: minor change to use UE_API definition

#Synced-CL 50160130

diff --git a/Engine/Plugins/VirtualProduction/SceneState/Source/SceneState/Public/SceneStatePlayer.h b/Engine/Plugins/VirtualProduction/SceneState/Source/SceneState/Public/SceneStatePlayer.h
--- a/Engine/Plugins/VirtualProduction/SceneState/Source/SceneState/Public/SceneStatePlayer.h
+++ b/Engine/Plugins/VirtualProduction/SceneState/Source/SceneState/Public/SceneStatePlayer.h
@@ -6,6 +6,8 @@
 #include "UObject/Object.h"
 #include "SceneStatePlayer.generated.h"
 
+#define UE_API SCENESTATE_API
+
 class USceneStateObject;
 
 /**
@@ -25,7 +27,7 @@ public:
 	}
 
 	/** Sets the class of the root state object. Destroying the existing one */
-	SCENESTATE_API void SetSceneStateClass(TSubclassOf<USceneStateObject> InSceneStateClass);
+	UE_API void SetSceneStateClass(TSubclassOf<USceneStateObject> InSceneStateClass);
 
 	USceneStateObject* GetSceneState() const
 	{
@@ -33,26 +35,26 @@ public:
 	}
 
 	/** Returns the context name for this player, for debugging purposes */
-	SCENESTATE_API FString GetContextName() const;
+	UE_API FString GetContextName() const;
 
 	/** Returns the context object for this player */
-	SCENESTATE_API UObject* GetContextObject() const;
+	UE_API UObject* GetContextObject() const;
 
 	/** Should be called prior to the start of the execution */
-	SCENESTATE_API void Setup();
+	UE_API void Setup();
 
 	/** Should be called to start the execution of the scene state object */
-	SCENESTATE_API void Begin();
+	UE_API void Begin();
 
 	/** Should be called every tick during the execution of the scene state object */
-	SCENESTATE_API void Tick(float InDeltaTime);
+	UE_API void Tick(float InDeltaTime);
 
 	/** Should be called to end the execution of the scene state object */
-	SCENESTATE_API void End();
+	UE_API void End();
 
 	//~ Begin UObject
 #if WITH_EDITOR
-	SCENESTATE_API virtual void PostEditChangeProperty(FPropertyChangedEvent& InPropertyChangedEvent) override;
+	UE_API virtual void PostEditChangeProperty(FPropertyChangedEvent& InPropertyChangedEvent) override;
 #endif
 	//~ End UObject
 
@@ -93,3 +95,5 @@ protected:
 	bool bEditableSceneStateClass = true;
 #endif
 };
+
+#undef UE_API
diff --git a/Engine/Plugins/VirtualProduction/SceneState/Source/SceneState/Public/SceneStatePropertyUtils.h b/Engine/Plugins/VirtualProduction/SceneState/Source/SceneState/Public/SceneStatePropertyUtils.h
--- a/Engine/Plugins/VirtualProduction/SceneState/Source/SceneState/Public/SceneStatePropertyUtils.h
+++ b/Engine/Plugins/VirtualProduction/SceneState/Source/SceneState/Public/SceneStatePropertyUtils.h
@@ -5,6 +5,8 @@
 #include "SceneStateBindingReference.h"
 #include "SceneStatePropertyReferenceUtils.h"
 
+#define UE_API SCENESTATE_API
+
 struct FSceneStateExecutionContext;
 struct FSceneStatePropertyReference;
 
@@ -24,7 +26,7 @@ namespace UE::SceneState
 	 * @param OutResult if successful, a result containing a valid resolved reference and a valid value address.
 	 * @return whether the operation succeeded
 	 */
-	SCENESTATE_API bool ResolveProperty(const FSceneStateExecutionContext& InContext, const FSceneStatePropertyReference& InPropertyReference, FResolvePropertyResult& OutResult);
+	UE_API bool ResolveProperty(const FSceneStateExecutionContext& InContext, const FSceneStatePropertyReference& InPropertyReference, FResolvePropertyResult& OutResult);
 
 	template<class T>
 	T* ResolveProperty(const FSceneStateExecutionContext& InContext, const FSceneStatePropertyReference& InPropertyReference)
@@ -38,3 +40,5 @@ namespace UE::SceneState
 	}
 
 } // UE::SceneState
+
+#undef UE_API
diff --git a/Engine/Plugins/VirtualProduction/SceneState/Source/SceneState/Public/Tasks/SceneStateBlueprintableTaskWrapper.h b/Engine/Plugins/VirtualProduction/SceneState/Source/SceneState/Public/Tasks/SceneStateBlueprintableTaskWrapper.h
--- a/Engine/Plugins/VirtualProduction/SceneState/Source/SceneState/Public/Tasks/SceneStateBlueprintableTaskWrapper.h
+++ b/Engine/Plugins/VirtualProduction/SceneState/Source/SceneState/Public/Tasks/SceneStateBlueprintableTaskWrapper.h
@@ -7,6 +7,8 @@
 #include "Templates/SubclassOf.h"
 #include "SceneStateBlueprintableTaskWrapper.generated.h"
 
+#define UE_API SCENESTATE_API
+
 class USceneStateBlueprintableTask;
 
 USTRUCT()
@@ -32,9 +34,9 @@ struct FSceneStateBlueprintableTaskWrapper : public FSceneStateTask
 
 	using FInstanceDataType = FSceneStateBlueprintableTaskInstance;
 
-	SCENESTATE_API FSceneStateBlueprintableTaskWrapper();
+	UE_API FSceneStateBlueprintableTaskWrapper();
 
-	SCENESTATE_API bool SetTaskClass(TSubclassOf<USceneStateBlueprintableTask> InTaskClass);
+	UE_API bool SetTaskClass(TSubclassOf<USceneStateBlueprintableTask> InTaskClass);
 
 	TSubclassOf<USceneStateBlueprintableTask> GetTaskClass() const
 	{
@@ -57,3 +59,5 @@ private:
 	UPROPERTY()
 	TSubclassOf<USceneStateBlueprintableTask> TaskClass;
 };
+
+#undef UE_API
diff --git a/Engine/Plugins/VirtualProduction/SceneState/Source/SceneState/Public/Tasks/SceneStateTask.h b/Engine/Plugins/VirtualProduction/SceneState/Source/SceneState/Public/Tasks/SceneStateTask.h
--- a/Engine/Plugins/VirtualProduction/SceneState/Source/SceneState/Public/Tasks/SceneStateTask.h
+++ b/Engine/Plugins/VirtualProduction/SceneState/Source/SceneState/Public/Tasks/SceneStateTask.h
@@ -7,6 +7,8 @@
 #include "SceneStateRange.h"
 #include "SceneStateTask.generated.h"
 
+#define UE_API SCENESTATE_API
+
 struct FSceneStateBindingDesc;
 struct FSceneStateExecutionContext;
 struct FSceneStateTaskBindingExtension;
@@ -43,55 +45,55 @@ struct FSceneStateTask
 	 * Called in-editor to get the task instance type
 	 * @return the task instance struct. Must derive from FSceneStateTaskInstance
 	 */
-	SCENESTATE_API const UScriptStruct* GetTaskInstanceType() const;
+	UE_API const UScriptStruct* GetTaskInstanceType() const;
 
 	/**
 	 * Called to init properties (e.g. Guids) or instances objects within the task instance
 	 * @param InOuter the outer to use for instanced objects
 	 * @param InTaskInstance the task instance to build
 	 */
-	SCENESTATE_API void BuildTaskInstance(UObject* InOuter, FStructView InTaskInstance) const;
+	UE_API void BuildTaskInstance(UObject* InOuter, FStructView InTaskInstance) const;
 
 	/**
 	 * Called when a property within a task or task instance has changed
 	 * @param InEditChange the edit change information (property changed event, outer)
 	 * @param InTaskInstance the task instance paired with the task
 	 */
-	SCENESTATE_API void PostEditChange(UE::SceneState::FTaskEditChange& InEditChange, FStructView InTaskInstance);
+	UE_API void PostEditChange(UE::SceneState::FTaskEditChange& InEditChange, FStructView InTaskInstance);
 #endif
 
 	/** Called to get the optional Binding Extension for a Task */
-	SCENESTATE_API const FSceneStateTaskBindingExtension* GetBindingExtension() const;
+	UE_API const FSceneStateTaskBindingExtension* GetBindingExtension() const;
 
 	/** Finds the Task Instance for the given Execution Instance */
-	SCENESTATE_API FStructView FindTaskInstance(const FSceneStateExecutionContext& InContext) const;
+	UE_API FStructView FindTaskInstance(const FSceneStateExecutionContext& InContext) const;
 
 	/**
 	 * Called when the State holding the Task first starts.
 	 * This is called for all tasks held by the state even if these tasks end up not running, or run at a later time.
 	 */
-	SCENESTATE_API void Setup(const FSceneStateExecutionContext& InContext, FStructView InTaskInstance) const;
+	UE_API void Setup(const FSceneStateExecutionContext& InContext, FStructView InTaskInstance) const;
 
 	/**
 	 * Called once the State processes that the Task has all its prerequisites met.
 	 * Applies and property bindings to the task and calls OnStart().
 	 */
-	SCENESTATE_API void Start(const FSceneStateExecutionContext& InContext, FStructView InTaskInstance) const;
+	UE_API void Start(const FSceneStateExecutionContext& InContext, FStructView InTaskInstance) const;
 
 	/**
 	 * Called by the State each Tick.
 	 * OnTick() gets called if the Task is running and set to Tickable (see TaskFlags)
 	 */
-	SCENESTATE_API void Tick(const FSceneStateExecutionContext& InContext, FStructView InTaskInstance, float InDeltaSeconds) const;
+	UE_API void Tick(const FSceneStateExecutionContext& InContext, FStructView InTaskInstance, float InDeltaSeconds) const;
 
 	/**
 	 * Called to Stop the Task, because it either finished or was forcibly stopped by the state
 	 * @see ESceneStateTaskStopReason
 	 */
-	SCENESTATE_API void Stop(const FSceneStateExecutionContext& InContext, FStructView InTaskInstance, ESceneStateTaskStopReason InStopReason) const;
+	UE_API void Stop(const FSceneStateExecutionContext& InContext, FStructView InTaskInstance, ESceneStateTaskStopReason InStopReason) const;
 
 	/** Calls Stop with stop reason being that the task has finished */
-	SCENESTATE_API void Finish(const FSceneStateExecutionContext& InContext, FStructView InTaskInstance) const;
+	UE_API void Finish(const FSceneStateExecutionContext& InContext, FStructView InTaskInstance) const;
 
 	uint16 GetParentStateIndex() const
 	{
@@ -109,9 +111,9 @@ struct FSceneStateTask
 	}
 
 protected:
-	SCENESTATE_API void SetFlags(ESceneStateTaskFlags InFlags);
+	UE_API void SetFlags(ESceneStateTaskFlags InFlags);
 
-	SCENESTATE_API void ClearFlags(ESceneStateTaskFlags InFlags);
+	UE_API void ClearFlags(ESceneStateTaskFlags InFlags);
 
 #if WITH_EDITOR
 	/**
@@ -204,3 +206,5 @@ private:
 	friend UE::SceneState::Editor::FStateMachineCompiler;
 	friend UE::SceneState::Editor::FStateMachineTaskCompiler;
 };
+
+#undef UE_API
diff --git a/Engine/Plugins/VirtualProduction/SceneState/Source/SceneState/Public/Tasks/SceneStateTaskExecutionContext.h b/Engine/Plugins/VirtualProduction/SceneState/Source/SceneState/Public/Tasks/SceneStateTaskExecutionContext.h
--- a/Engine/Plugins/VirtualProduction/SceneState/Source/SceneState/Public/Tasks/SceneStateTaskExecutionContext.h
+++ b/Engine/Plugins/VirtualProduction/SceneState/Source/SceneState/Public/Tasks/SceneStateTaskExecutionContext.h
@@ -6,6 +6,8 @@
 #include "SceneStateRange.h"
 #include "Templates/SharedPointer.h"
 
+#define UE_API SCENESTATE_API
+
 struct FConstStructView;
 struct FSceneStateExecutionContext;
 struct FSceneStateTask;
@@ -26,21 +28,21 @@ namespace UE::SceneState
  */
 struct FTaskExecutionContext
 {
-	SCENESTATE_API FTaskExecutionContext();
+	UE_API FTaskExecutionContext();
 
-	SCENESTATE_API explicit FTaskExecutionContext(const FSceneStateTask& InTask, const FSceneStateExecutionContext& InContext);
+	UE_API explicit FTaskExecutionContext(const FSceneStateTask& InTask, const FSceneStateExecutionContext& InContext);
 
 	/** Resolves the execution context if it still exists */
-	SCENESTATE_API const FSceneStateExecutionContext* GetExecutionContext() const;
+	UE_API const FSceneStateExecutionContext* GetExecutionContext() const;
 
 	/** Retrieves the FSceneStateTask view for this context, if the task instance this context got created for is still valid. An invalid view is returned otherwise */
-	SCENESTATE_API FConstStructView GetTask() const;
+	UE_API FConstStructView GetTask() const;
 
 	/** Retrieves the FSceneStateTaskInstance view for this context, if the task instance this context got created for is still valid. An invalid view is returned otherwise */
-	SCENESTATE_API FStructView GetTaskInstance() const;
+	UE_API FStructView GetTaskInstance() const;
 
 	/** Helper function to Finish the Task / Task Instance for this context if these are still valid */
-	SCENESTATE_API void FinishTask() const;
+	UE_API void FinishTask() const;
 
 private:
 	FString GetDebugString() const;
@@ -66,3 +68,5 @@ private:
 };
 
 } // UE::SceneState
+
+#undef UE_API
diff --git a/Engine/Plugins/VirtualProduction/SceneState/Source/SceneStateBinding/Public/SceneStateBinding.h b/Engine/Plugins/VirtualProduction/SceneState/Source/SceneStateBinding/Public/SceneStateBinding.h
--- a/Engine/Plugins/VirtualProduction/SceneState/Source/SceneStateBinding/Public/SceneStateBinding.h
+++ b/Engine/Plugins/VirtualProduction/SceneState/Source/SceneStateBinding/Public/SceneStateBinding.h
@@ -6,6 +6,8 @@
 #include "SceneStateBindingDataHandle.h"
 #include "SceneStateBinding.generated.h"
 
+#define UE_API SCENESTATEBINDING_API
+
 /** Represents a binding from a source property to a destination property */
 USTRUCT()
 struct FSceneStateBinding : public FPropertyBindingBinding
@@ -15,7 +17,7 @@ struct FSceneStateBinding : public FPropertyBindingBinding
 	using FPropertyBindingBinding::FPropertyBindingBinding;
 
 	//~ Begin FPropertyBindingBinding
-	SCENESTATEBINDING_API virtual FConstStructView GetSourceDataHandleStruct() const override;
+	UE_API virtual FConstStructView GetSourceDataHandleStruct() const override;
 	//~ End FPropertyBindingBinding
 
 	/** Handle to the data containing the source property */
@@ -26,3 +28,5 @@ struct FSceneStateBinding : public FPropertyBindingBinding
 	UPROPERTY()
 	FSceneStateBindingDataHandle TargetDataHandle;
 };
+
+#undef UE_API
diff --git a/Engine/Plugins/VirtualProduction/SceneState/Source/SceneStateBinding/Public/SceneStateBindingDataHandle.h b/Engine/Plugins/VirtualProduction/SceneState/Source/SceneStateBinding/Public/SceneStateBindingDataHandle.h
--- a/Engine/Plugins/VirtualProduction/SceneState/Source/SceneStateBinding/Public/SceneStateBindingDataHandle.h
+++ b/Engine/Plugins/VirtualProduction/SceneState/Source/SceneStateBinding/Public/SceneStateBindingDataHandle.h
@@ -4,6 +4,8 @@
 
 #include "SceneStateBindingDataHandle.generated.h"
 
+#define UE_API SCENESTATEBINDING_API
+
 UENUM()
 enum class ESceneStateDataType : uint8
 {
@@ -29,7 +31,7 @@ struct FSceneStateBindingDataHandle
 		int32 Index = INDEX_NONE;
 		int32 SubIndex = INDEX_NONE;
 	};
-	SCENESTATEBINDING_API static FSceneStateBindingDataHandle MakeExternalDataHandle(const FExternalData& InExternalData);
+	UE_API static FSceneStateBindingDataHandle MakeExternalDataHandle(const FExternalData& InExternalData);
 
 	explicit FSceneStateBindingDataHandle(ESceneStateDataType InDataType, int32 InDataIndex = INDEX_NONE, int32 InDataSubIndex = INDEX_NONE)
 		: Type(static_cast<uint16>(InDataType))
@@ -112,3 +114,5 @@ private:
 	UPROPERTY()
 	uint16 SubIndex = InvalidIndex;
 };
+
+#undef UE_API
diff --git a/Engine/Plugins/VirtualProduction/SceneState/Source/SceneStateBinding/Public/SceneStateBindingDelegates.h b/Engine/Plugins/VirtualProduction/SceneState/Source/SceneStateBinding/Public/SceneStateBindingDelegates.h
--- a/Engine/Plugins/VirtualProduction/SceneState/Source/SceneStateBinding/Public/SceneStateBindingDelegates.h
+++ b/Engine/Plugins/VirtualProduction/SceneState/Source/SceneStateBinding/Public/SceneStateBindingDelegates.h
@@ -6,6 +6,8 @@
 #include "Misc/Guid.h"
 #include "UObject/ObjectPtr.h"
 
+#define UE_API SCENESTATEBINDING_API
+
 class UObject;
 
 namespace UE::SceneState
@@ -20,7 +22,9 @@ namespace UE::SceneState
 		TMap<FGuid, FGuid> OldToNewStructIdMap;
 	};
 	/** Delegate called when a struct id has changed */
-	extern SCENESTATEBINDING_API TMulticastDelegate<void(const FStructIdChange&)> OnStructIdChanged;
+	extern UE_API TMulticastDelegate<void(const FStructIdChange&)> OnStructIdChanged;
 #endif
 
 } // UE::SceneState
+
+#undef UE_API
diff --git a/Engine/Plugins/VirtualProduction/SceneState/Source/SceneStateBinding/Public/SceneStateBindingUtils.h b/Engine/Plugins/VirtualProduction/SceneState/Source/SceneStateBinding/Public/SceneStateBindingUtils.h
--- a/Engine/Plugins/VirtualProduction/SceneState/Source/SceneStateBinding/Public/SceneStateBindingUtils.h
+++ b/Engine/Plugins/VirtualProduction/SceneState/Source/SceneStateBinding/Public/SceneStateBindingUtils.h
@@ -5,6 +5,8 @@
 #include "PropertyBindingDataView.h"
 #include "Templates/Function.h"
 
+#define UE_API SCENESTATEBINDING_API
+
 class UObject;
 struct FGuid;
 struct FSceneStateBindingCollection;
@@ -21,7 +23,7 @@ namespace UE::SceneState
 	 * @param InOldStructId the struct id to queue copy bindings from
 	 * @param InNewStructId the struct id to queue copy bindings to
 	 */
-	SCENESTATEBINDING_API void HandleStructIdChanged(UObject& InObject, const FGuid& InOldStructId, const FGuid& InNewStructId);
+	UE_API void HandleStructIdChanged(UObject& InObject, const FGuid& InOldStructId, const FGuid& InNewStructId);
 #endif
 
 	struct FPatchBindingParams
@@ -32,7 +34,7 @@ namespace UE::SceneState
 		TFunctionRef<const UStruct*(const FSceneStateBindingDataHandle&)> FindDataStructFunctor;
 	};
 	/** Patches invalidated structs (property bags, user defined structs, etc.) in the given binding collection */
-	SCENESTATEBINDING_API void PatchBindingCollection(const FPatchBindingParams& InParams);
+	UE_API void PatchBindingCollection(const FPatchBindingParams& InParams);
 
 	/** Patches any invalidated structs (property bags, user defined structs, etc.) in the binding descs of the given binding collection */
 	void PatchBindingDescs(const FPatchBindingParams& InParams);
@@ -44,3 +46,5 @@ namespace UE::SceneState
 	void PatchCopyBatches(const FPatchBindingParams& InParams);
 
 } // UE::SceneState
+
+#undef UE_API
diff --git a/Engine/Plugins/VirtualProduction/SceneState/Source/SceneStateBinding/Public/SceneStatePropertyReferenceUtils.h b/Engine/Plugins/VirtualProduction/SceneState/Source/SceneStateBinding/Public/SceneStatePropertyReferenceUtils.h
--- a/Engine/Plugins/VirtualProduction/SceneState/Source/SceneStateBinding/Public/SceneStatePropertyReferenceUtils.h
+++ b/Engine/Plugins/VirtualProduction/SceneState/Source/SceneStateBinding/Public/SceneStatePropertyReferenceUtils.h
@@ -8,26 +8,28 @@
 #include "UObject/TextProperty.h"
 #include "UObject/UnrealType.h"
 
+#define UE_API SCENESTATEBINDING_API
+
 struct FEdGraphPinType;
 struct FSceneStateBlueprintPropertyReference;
 
 namespace UE::SceneState
 {
 	/** Returns whether the given property is a struct property of struct FSceneStatePropertyReference */
-	SCENESTATEBINDING_API bool IsPropertyReference(const FProperty* InProperty);
+	UE_API bool IsPropertyReference(const FProperty* InProperty);
 
 #if WITH_EDITOR
 	/** Returns whether a source property is compatible with the reference property */
-	SCENESTATEBINDING_API bool IsPropertyReferenceCompatible(const FProperty* InReferenceProperty, const void* InReferenceAddress, const FProperty* InSourceProperty, const void* InSourceAddress);
+	UE_API bool IsPropertyReferenceCompatible(const FProperty* InReferenceProperty, const void* InReferenceAddress, const FProperty* InSourceProperty, const void* InSourceAddress);
 
 	/** Returns the first graph pin type based on a given Property Reference property's meta-data */
-	SCENESTATEBINDING_API FEdGraphPinType GetPropertyReferencePinType(const FProperty* InProperty, const void* InReferenceAddress);
+	UE_API FEdGraphPinType GetPropertyReferencePinType(const FProperty* InProperty, const void* InReferenceAddress);
 
 	/** Returns the graph pin type based on a given Property Reference */
-	SCENESTATEBINDING_API TArray<FEdGraphPinType, TInlineAllocator<1>> GetPropertyReferencePinTypes(const FSceneStateBlueprintPropertyReference& InPropertyReference);
+	UE_API TArray<FEdGraphPinType, TInlineAllocator<1>> GetPropertyReferencePinTypes(const FSceneStateBlueprintPropertyReference& InPropertyReference);
 
 	/** Returns the graph pin types based on a given Property Reference property's meta-data */
-	SCENESTATEBINDING_API TArray<FEdGraphPinType, TInlineAllocator<1>> GetPropertyReferencePinTypes(const FProperty* InProperty, const void* InReferenceAddress);
+	UE_API TArray<FEdGraphPinType, TInlineAllocator<1>> GetPropertyReferencePinTypes(const FProperty* InProperty, const void* InReferenceAddress);
 #endif
 
 	template<class InDataType>
@@ -165,3 +167,5 @@ namespace UE::SceneState
 		}
 	};
 }
+
+#undef UE_API
diff --git a/Engine/Plugins/VirtualProduction/SceneState/Source/SceneStateBlueprint/Public/SceneStateBlueprintDelegates.h b/Engine/Plugins/VirtualProduction/SceneState/Source/SceneStateBlueprint/Public/SceneStateBlueprintDelegates.h
--- a/Engine/Plugins/VirtualProduction/SceneState/Source/SceneStateBlueprint/Public/SceneStateBlueprintDelegates.h
+++ b/Engine/Plugins/VirtualProduction/SceneState/Source/SceneStateBlueprint/Public/SceneStateBlueprintDelegates.h
@@ -5,6 +5,8 @@
 #include "Delegates/Delegate.h"
 #include "UObject/ObjectPtr.h"
 
+#define UE_API SCENESTATEBLUEPRINT_API
+
 class UObject;
 class USceneStateBlueprint;
 
@@ -25,6 +27,8 @@ struct FBlueprintDebugObjectChange
  * UBlueprint::DebuggingWorldRegistrationHelper is closer to what is needed, but still requires a valid UWorld to have been found.
  * This delegate solves this by broadcasting in USceneStateBlueprint::SetObjectBeingDebugged without requiring a valid world.
  */
-extern SCENESTATEBLUEPRINT_API TMulticastDelegate<void(const FBlueprintDebugObjectChange&)> OnBlueprintDebugObjectChanged;
+extern UE_API TMulticastDelegate<void(const FBlueprintDebugObjectChange&)> OnBlueprintDebugObjectChanged;
 
 } // UE::SceneState
+
+#undef UE_API
diff --git a/Engine/Plugins/VirtualProduction/SceneState/Source/SceneStateBlueprintEditor/Public/SceneStateBlueprintEditorUtils.h b/Engine/Plugins/VirtualProduction/SceneState/Source/SceneStateBlueprintEditor/Public/SceneStateBlueprintEditorUtils.h
--- a/Engine/Plugins/VirtualProduction/SceneState/Source/SceneStateBlueprintEditor/Public/SceneStateBlueprintEditorUtils.h
+++ b/Engine/Plugins/VirtualProduction/SceneState/Source/SceneStateBlueprintEditor/Public/SceneStateBlueprintEditorUtils.h
@@ -5,6 +5,8 @@
 #include "PropertyEditorModule.h"
 #include "Templates/SharedPointerFwd.h"
 
+#define UE_API SCENESTATEBLUEPRINTEDITOR_API
+
 class IDetailChildrenBuilder;
 class IPropertyHandle;
 class UClass;
@@ -14,7 +16,7 @@ struct FInstancedPropertyBag;
 namespace UE::SceneState::Editor
 {
 	/** Finds the Task Id for a given Property Handle by first finding the outer Task Node and returning its Task id. */
-	SCENESTATEBLUEPRINTEDITOR_API FGuid FindTaskId(const TSharedRef<IPropertyHandle>& InPropertyHandle);
+	UE_API FGuid FindTaskId(const TSharedRef<IPropertyHandle>& InPropertyHandle);
 
 	/**
 	 * Returns whether a given property is an object property of a given class
@@ -22,7 +24,7 @@ namespace UE::SceneState::Editor
 	 * @param InClass the class to check
 	 * @return true if the property is an object property of the given class
 	 */
-	SCENESTATEBLUEPRINTEDITOR_API bool IsObjectPropertyOfClass(const FProperty* InProperty, const UClass* InClass);
+	UE_API bool IsObjectPropertyOfClass(const FProperty* InProperty, const UClass* InClass);
 
 	/**
 	 * Returns whether a given property is of a given struct type
@@ -30,7 +32,7 @@ namespace UE::SceneState::Editor
 	 * @param InStruct the type to check
 	 * @return true if the property is a struct property of the given struct type  
 	 */
-	SCENESTATEBLUEPRINTEDITOR_API bool IsStruct(const FProperty* InProperty, const UScriptStruct* InStruct);
+	UE_API bool IsStruct(const FProperty* InProperty, const UScriptStruct* InStruct);
 
 	/**
 	 * Gets the Guid value from the Property Handle.
@@ -38,14 +40,14 @@ namespace UE::SceneState::Editor
 	 * @param OutGuid the returned guid if the operation was successful
 	 * @return whether the guid was obtained successfully, or whether it failed or had multiple values
 	 */
-	SCENESTATEBLUEPRINTEDITOR_API FPropertyAccess::Result GetGuid(const TSharedRef<IPropertyHandle>& InGuidPropertyHandle, FGuid& OutGuid);
+	UE_API FPropertyAccess::Result GetGuid(const TSharedRef<IPropertyHandle>& InGuidPropertyHandle, FGuid& OutGuid);
 
 	/**
 	 * Tries to create a structure data provider for the given instanced struct property handle
 	 * @param InStructHandle the instanced struct property handle to create a struct data provider for
 	 * @return the created struct data provider
 	 */
-	SCENESTATEBLUEPRINTEDITOR_API TSharedRef<IStructureDataProvider> CreateInstancedStructDataProvider(const TSharedRef<IPropertyHandle>& InStructHandle);
+	UE_API TSharedRef<IStructureDataProvider> CreateInstancedStructDataProvider(const TSharedRef<IPropertyHandle>& InStructHandle);
 
 	/**
 	 * Compares the structure layout of two property bags
@@ -53,13 +55,13 @@ namespace UE::SceneState::Editor
 	 * @param InParametersB the second property bag comparand
 	 * @return true if both have the same number of properties with same names and compatible types.
 	 */
-	SCENESTATEBLUEPRINTEDITOR_API bool CompareParametersLayout(const FInstancedPropertyBag& InParametersA, const FInstancedPropertyBag& InParametersB);
+	UE_API bool CompareParametersLayout(const FInstancedPropertyBag& InParametersA, const FInstancedPropertyBag& InParametersB);
 
 	/** Sets the Instance Meta-data of a given property handle to a Task Id guid */
-	SCENESTATEBLUEPRINTEDITOR_API void AssignBindingId(const TSharedRef<IPropertyHandle>& InPropertyHandle, const FGuid& InTaskId);
+	UE_API void AssignBindingId(const TSharedRef<IPropertyHandle>& InPropertyHandle, const FGuid& InTaskId);
 
 	/** Finds the Common Class for all the Objects within the Object Property Handle */
-	SCENESTATEBLUEPRINTEDITOR_API UClass* FindCommonBase(const TSharedRef<IPropertyHandle>& InPropertyHandle);
+	UE_API UClass* FindCommonBase(const TSharedRef<IPropertyHandle>& InPropertyHandle);
 
 	/**
 	 * Helper function to add the Common properties of the object property handle to the Child Builder
@@ -67,6 +69,8 @@ namespace UE::SceneState::Editor
 	 * @param InChildBuilder the builder to add the properties to
 	 * @param InDisallowedFlags filter flags. Any property with any of these flags will be skipped 
 	 */
-	SCENESTATEBLUEPRINTEDITOR_API bool AddObjectProperties(const TSharedRef<IPropertyHandle>& InPropertyHandle, IDetailChildrenBuilder& InChildBuilder, uint64 InDisallowedFlags = CPF_DisableEditOnInstance);
+	UE_API bool AddObjectProperties(const TSharedRef<IPropertyHandle>& InPropertyHandle, IDetailChildrenBuilder& InChildBuilder, uint64 InDisallowedFlags = CPF_DisableEditOnInstance);
 
 } // UE::SceneState::Editor
+
+#undef UE_API
diff --git a/Engine/Plugins/VirtualProduction/SceneState/Source/SceneStateEditor/Public/SceneStateEditorUtils.h b/Engine/Plugins/VirtualProduction/SceneState/Source/SceneStateEditor/Public/SceneStateEditorUtils.h
--- a/Engine/Plugins/VirtualProduction/SceneState/Source/SceneStateEditor/Public/SceneStateEditorUtils.h
+++ b/Engine/Plugins/VirtualProduction/SceneState/Source/SceneStateEditor/Public/SceneStateEditorUtils.h
@@ -4,6 +4,8 @@
 
 #include "HAL/Platform.h"
 
+#define UE_API SCENESTATEEDITOR_API
+
 class FText;
 class UStruct;
 
@@ -14,6 +16,8 @@ namespace UE::SceneState::Editor
 	 * This is different from UStruct::GetTooltipText as it will return an empty text if there was no explicit Tooltip found
 	 * @return the found tooltip, or empty if not the metadata was not present.
 	 */
-	SCENESTATEEDITOR_API FText GetStructTooltip(const UStruct& InStruct);
+	UE_API FText GetStructTooltip(const UStruct& InStruct);
 
 } // UE::SceneState::Editor
+
+#undef UE_API
diff --git a/Engine/Plugins/VirtualProduction/SceneState/Source/SceneStateEditor/Public/Tasks/SceneStateTaskDesc.h b/Engine/Plugins/VirtualProduction/SceneState/Source/SceneStateEditor/Public/Tasks/SceneStateTaskDesc.h
--- a/Engine/Plugins/VirtualProduction/SceneState/Source/SceneStateEditor/Public/Tasks/SceneStateTaskDesc.h
+++ b/Engine/Plugins/VirtualProduction/SceneState/Source/SceneStateEditor/Public/Tasks/SceneStateTaskDesc.h
@@ -5,6 +5,8 @@
 #include "StructUtils/StructView.h"
 #include "SceneStateTaskDesc.generated.h"
 
+#define UE_API SCENESTATEEDITOR_API
+
 struct FSceneStateTask;
 struct FSceneStateTaskInstance;
 
@@ -44,7 +46,7 @@ struct FSceneStateTaskDesc
 {
 	GENERATED_BODY()
 
-	SCENESTATEEDITOR_API FSceneStateTaskDesc();
+	UE_API FSceneStateTaskDesc();
 
 	virtual ~FSceneStateTaskDesc() = default;
 
@@ -55,16 +57,16 @@ struct FSceneStateTaskDesc
 	}
 
 	/** Gets the display name of the task */
-	SCENESTATEEDITOR_API FText GetDisplayName(const FSceneStateTaskDescContext& InContext) const;
+	UE_API FText GetDisplayName(const FSceneStateTaskDescContext& InContext) const;
 
 	/** Gets the tooltip of the task */
-	SCENESTATEEDITOR_API FText GetTooltip(const FSceneStateTaskDescContext& InContext) const;
+	UE_API FText GetTooltip(const FSceneStateTaskDescContext& InContext) const;
 
 	/** Gets the target object to go to on node double click */
-	SCENESTATEEDITOR_API UObject* GetJumpTarget(const FSceneStateTaskDescContext& InContext) const;
+	UE_API UObject* GetJumpTarget(const FSceneStateTaskDescContext& InContext) const;
 
 	/** Called when struct ids have been changed, to be handled by tasks individually */
-	SCENESTATEEDITOR_API void NotifyStructIdsChanged(const FSceneStateTaskDescMutableContext& InContext, const UE::SceneState::FStructIdChange& InChange) const;
+	UE_API void NotifyStructIdsChanged(const FSceneStateTaskDescMutableContext& InContext, const UE::SceneState::FStructIdChange& InChange) const;
 
 	/** Set the task struct to support */
 	template<typename T UE_REQUIRES(std::is_base_of_v<FSceneStateTask, T>)>
@@ -105,9 +107,11 @@ private:
 	bool IsValidContext(const FSceneStateTaskDescMutableContext& InContext) const;
 
 	/** Set the task struct to support */
-	SCENESTATEEDITOR_API void SetSupportedTask(UScriptStruct* InSupportedTask);
+	UE_API void SetSupportedTask(UScriptStruct* InSupportedTask);
 
 	/** The task struct to support */
 	UPROPERTY()
 	TObjectPtr<const UScriptStruct> SupportedTask = nullptr;
 };
+
+#undef UE_API
diff --git a/Engine/Plugins/VirtualProduction/SceneState/Source/SceneStateEditor/Public/Tasks/SceneStateTaskDescRegistry.h b/Engine/Plugins/VirtualProduction/SceneState/Source/SceneStateEditor/Public/Tasks/SceneStateTaskDescRegistry.h
--- a/Engine/Plugins/VirtualProduction/SceneState/Source/SceneStateEditor/Public/Tasks/SceneStateTaskDescRegistry.h
+++ b/Engine/Plugins/VirtualProduction/SceneState/Source/SceneStateEditor/Public/Tasks/SceneStateTaskDescRegistry.h
@@ -8,6 +8,8 @@
 #include "UObject/GCObject.h"
 #include "UObject/ObjectKey.h"
 
+#define UE_API SCENESTATEEDITOR_API
+
 /** Registry holding editor-only task descs for each discovered task struct */
 class FSceneStateTaskDescRegistry : public FGCObject
 {
@@ -15,10 +17,10 @@ public:
 	static FSceneStateTaskDescRegistry GlobalRegistry;
 
 	/** Gets a const ref to the global registry */
-	SCENESTATEEDITOR_API static const FSceneStateTaskDescRegistry& Get();
+	UE_API static const FSceneStateTaskDescRegistry& Get();
 
 	/** Gets the task desc for the given task struct */
-	SCENESTATEEDITOR_API const FSceneStateTaskDesc& GetTaskDesc(const UScriptStruct* InTaskStruct) const;
+	UE_API const FSceneStateTaskDesc& GetTaskDesc(const UScriptStruct* InTaskStruct) const;
 
 	/** Gathers all the available Tasks and stores it in this registry */
 	void CacheTaskDescs();
@@ -36,3 +38,5 @@ private:
 	/** Default task desc to use */
 	TInstancedStruct<FSceneStateTaskDesc> DefaultTaskDesc;
 };
+
+#undef UE_API
diff --git a/Engine/Plugins/VirtualProduction/SceneState/Source/SceneStateEvent/Public/SceneStateEventSchemaHandle.h b/Engine/Plugins/VirtualProduction/SceneState/Source/SceneStateEvent/Public/SceneStateEventSchemaHandle.h
--- a/Engine/Plugins/VirtualProduction/SceneState/Source/SceneStateEvent/Public/SceneStateEventSchemaHandle.h
+++ b/Engine/Plugins/VirtualProduction/SceneState/Source/SceneStateEvent/Public/SceneStateEventSchemaHandle.h
@@ -5,6 +5,8 @@
 #include "UObject/SoftObjectPtr.h"
 #include "SceneStateEventSchemaHandle.generated.h"
 
+#define UE_API SCENESTATEEVENT_API
+
 class USceneStateEventSchemaObject;
 class UUserDefinedStruct;
 struct FStructView;
@@ -16,13 +18,13 @@ struct FSceneStateEventSchemaHandle
 	GENERATED_BODY()
 
 	/** Get the struct of the Event Schema. Can be null if it's an event with no parameters */
-	SCENESTATEEVENT_API UUserDefinedStruct* GetEventStruct() const;
+	UE_API UUserDefinedStruct* GetEventStruct() const;
 
 	/** Resolves the event schema soft object pointer */
-	SCENESTATEEVENT_API USceneStateEventSchemaObject* GetEventSchema() const;
+	UE_API USceneStateEventSchemaObject* GetEventSchema() const;
 
 	/** Get the default data view of the Event Schema struct. Can be invalid if it's an event with no parameters */
-	SCENESTATEEVENT_API FStructView GetDefaultDataView() const;
+	UE_API FStructView GetDefaultDataView() const;
 
 	bool operator==(const FSceneStateEventSchemaHandle& InOther) const
 	{
@@ -48,3 +50,5 @@ private:
 	UPROPERTY(EditAnywhere, Category="Event")
 	TSoftObjectPtr<UUserDefinedStruct> EventStruct;
 };
+
+#undef UE_API
diff --git a/Engine/Plugins/VirtualProduction/SceneState/Source/SceneStateEvent/Public/SceneStateEventStream.h b/Engine/Plugins/VirtualProduction/SceneState/Source/SceneStateEvent/Public/SceneStateEventStream.h
--- a/Engine/Plugins/VirtualProduction/SceneState/Source/SceneStateEvent/Public/SceneStateEventStream.h
+++ b/Engine/Plugins/VirtualProduction/SceneState/Source/SceneStateEvent/Public/SceneStateEventStream.h
@@ -8,6 +8,8 @@
 #include "UObject/Object.h"
 #include "SceneStateEventStream.generated.h"
 
+#define UE_API SCENESTATEEVENT_API
+
 struct FSceneStateEventHandler;
 struct FSceneStateEventSchemaHandle;
 
@@ -19,31 +21,31 @@ class USceneStateEventStream : public UObject
 
 public:
 	/** Registers the Event Stream to the Event Subsystem to listen to Broadcast Events */
-	SCENESTATEEVENT_API bool Register();
+	UE_API bool Register();
 
 	/** Registers the Event Stream from the Event Subsystem */
-	SCENESTATEEVENT_API void Unregister();
+	UE_API void Unregister();
 
 	/** Pushes a new event into the stream */
-	SCENESTATEEVENT_API void PushEvent(FSharedStruct&& InEvent);
+	UE_API void PushEvent(FSharedStruct&& InEvent);
 
 	/** Pushes a new event into the stream */
-	SCENESTATEEVENT_API void PushEvent(const FSharedStruct& InEvent);
+	UE_API void PushEvent(const FSharedStruct& InEvent);
 
 	/** Consumes the first (oldest) Event that was pushed that matches the given Schema */
-	SCENESTATEEVENT_API bool ConsumeEventBySchema(const FSceneStateEventSchemaHandle& InEventSchemaHandle);
+	UE_API bool ConsumeEventBySchema(const FSceneStateEventSchemaHandle& InEventSchemaHandle);
 
 	/** Finds the first (oldest) Event that was pushed that matches the given Schema */
-	SCENESTATEEVENT_API const FSceneStateEvent* FindEventBySchema(const FSceneStateEventSchemaHandle& InEventSchemaHandle) const;
+	UE_API const FSceneStateEvent* FindEventBySchema(const FSceneStateEventSchemaHandle& InEventSchemaHandle) const;
 
 	/** Finds the event that is captured by the Handler with the given Id */
-	SCENESTATEEVENT_API FSceneStateEvent* FindCapturedEvent(const FGuid& InHandlerId);
+	UE_API FSceneStateEvent* FindCapturedEvent(const FGuid& InHandlerId);
 
 	/** Removes the Events (first/oldest) that match the given Event Handlers and moves them to the Captured Event map */
-	SCENESTATEEVENT_API void CaptureEvents(TConstArrayView<FSceneStateEventHandler> InEventHandlers);
+	UE_API void CaptureEvents(TConstArrayView<FSceneStateEventHandler> InEventHandlers);
 
 	/** Cleanup the Captured Events that match the given Event Handlers */
-	SCENESTATEEVENT_API void ResetCapturedEvents(TConstArrayView<FSceneStateEventHandler> InEventHandlers);
+	UE_API void ResetCapturedEvents(TConstArrayView<FSceneStateEventHandler> InEventHandlers);
 
 private:
 	/** Returns the index of the first event matching the event schema */
@@ -57,3 +59,5 @@ private:
 	UPROPERTY(Transient, DuplicateTransient, TextExportTransient)
 	TMap<FGuid, FSharedStruct> CapturedEvents;
 };
+
+#undef UE_API
diff --git a/Engine/Plugins/VirtualProduction/SceneState/Source/SceneStateEvent/Public/SceneStateEventTemplate.h b/Engine/Plugins/VirtualProduction/SceneState/Source/SceneStateEvent/Public/SceneStateEventTemplate.h
--- a/Engine/Plugins/VirtualProduction/SceneState/Source/SceneStateEvent/Public/SceneStateEventTemplate.h
+++ b/Engine/Plugins/VirtualProduction/SceneState/Source/SceneStateEvent/Public/SceneStateEventTemplate.h
@@ -6,6 +6,8 @@
 #include "StructUtils/InstancedStruct.h"
 #include "SceneStateEventTemplate.generated.h"
 
+#define UE_API SCENESTATEEVENT_API
+
 /** Defines a handle to an event schema and the payload event data to copy to events to push */
 USTRUCT(BlueprintType)
 struct FSceneStateEventTemplate
@@ -24,7 +26,7 @@ struct FSceneStateEventTemplate
 
 #if WITH_EDITOR
 	/** Syncs the event data to match the event schema struct */
-	SCENESTATEEVENT_API void SyncEventData();
+	UE_API void SyncEventData();
 #endif
 
 	static FName GetEventSchemaHandlePropertyName()
@@ -46,3 +48,5 @@ private:
 	UPROPERTY(EditAnywhere, Category="Event", meta=(StructTypeConst))
 	FInstancedStruct EventData;
 };
+
+#undef UE_API
diff --git a/Engine/Plugins/VirtualProduction/SceneState/Source/SceneStateEvent/Public/SceneStateEventUtils.h b/Engine/Plugins/VirtualProduction/SceneState/Source/SceneStateEvent/Public/SceneStateEventUtils.h
--- a/Engine/Plugins/VirtualProduction/SceneState/Source/SceneStateEvent/Public/SceneStateEventUtils.h
+++ b/Engine/Plugins/VirtualProduction/SceneState/Source/SceneStateEvent/Public/SceneStateEventUtils.h
@@ -11,6 +11,8 @@ struct FInstancedStruct;
 struct FSceneStateEventSchemaHandle;
 struct FSceneStateEventTemplate;
 
+#define UE_API SCENESTATEEVENT_API
+
 namespace UE::SceneState
 {
 	/**
@@ -20,7 +22,7 @@ namespace UE::SceneState
 	 * @param InEventData the event payload. Must match the event schema struct
 	 * @return true if the event was pushed successfully, false otherwise
 	 */
-	SCENESTATEEVENT_API bool PushEvent(USceneStateEventStream* InEventStream, const FSceneStateEventSchemaHandle& InEventSchemaHandle, FInstancedStruct&& InEventData);
+	UE_API bool PushEvent(USceneStateEventStream* InEventStream, const FSceneStateEventSchemaHandle& InEventSchemaHandle, FInstancedStruct&& InEventData);
 
 	/**
 	 * Pushes an event to the given event stream
@@ -28,14 +30,14 @@ namespace UE::SceneState
 	 * @param InEventTemplate the template used to create the event
 	 * @return true if the event was pushed successfully, false otherwise
 	 */
-	SCENESTATEEVENT_API bool PushEvent(USceneStateEventStream* InEventStream, const FSceneStateEventTemplate& InEventTemplate);
+	UE_API bool PushEvent(USceneStateEventStream* InEventStream, const FSceneStateEventTemplate& InEventTemplate);
 
 	/**
 	 * Retrieves a context world for a given object.
 	 * @param InObject the context object to look a context world for
 	 * @return the context world or GWorld as fallback if no explicit world was found.
 	 */
-	SCENESTATEEVENT_API const UWorld* GetContextWorld(const UObject* InObject);
+	UE_API const UWorld* GetContextWorld(const UObject* InObject);
 
 	/**
 	 * Broadcasts an event to the every registered event stream within a given scope.
@@ -44,7 +46,7 @@ namespace UE::SceneState
 	 * @param InEventData the event payload. Must match the event schema struct
 	 * @return true if the event was broadcasted successfully, false otherwise
 	 */
-	SCENESTATEEVENT_API bool BroadcastEvent(const UObject* InEventContext, const FSceneStateEventSchemaHandle& InEventSchemaHandle, FInstancedStruct&& InEventData);
+	UE_API bool BroadcastEvent(const UObject* InEventContext, const FSceneStateEventSchemaHandle& InEventSchemaHandle, FInstancedStruct&& InEventData);
 
 	/**
 	 * Broadcasts an event to the every registered event stream within a given scope.
@@ -52,6 +54,8 @@ namespace UE::SceneState
 	 * @param InEventTemplate the template used to create the event
 	 * @return true if the event was broadcasted successfully, false otherwise
 	 */
-	SCENESTATEEVENT_API bool BroadcastEvent(const UObject* InEventContext, const FSceneStateEventTemplate& InEventTemplate);
+	UE_API bool BroadcastEvent(const UObject* InEventContext, const FSceneStateEventTemplate& InEventTemplate);
 
 } // UE::SceneState
+
+#undef UE_API
diff --git a/Engine/Plugins/VirtualProduction/SceneState/Source/SceneStateGameplay/Public/SceneStateComponentPlayer.h b/Engine/Plugins/VirtualProduction/SceneState/Source/SceneStateGameplay/Public/SceneStateComponentPlayer.h
--- a/Engine/Plugins/VirtualProduction/SceneState/Source/SceneStateGameplay/Public/SceneStateComponentPlayer.h
+++ b/Engine/Plugins/VirtualProduction/SceneState/Source/SceneStateGameplay/Public/SceneStateComponentPlayer.h
@@ -5,6 +5,9 @@
 #include "SceneStatePlayer.h"
 #include "SceneStateComponentPlayer.generated.h"
 
+#define UE_API SCENESTATEGAMEPLAY_API
+
+class AActor;
 class USceneStateComponent;
 
 /** Scene State Player for Scene State Components */
@@ -14,10 +17,12 @@ class USceneStateComponentPlayer : public USceneStatePlayer
 	GENERATED_BODY()
 
 protected:
-	SCENESTATEGAMEPLAY_API AActor* GetActor() const;
+	UE_API AActor* GetActor() const;
 
 	//~ Begin USceneStatePlayer
-	SCENESTATEGAMEPLAY_API virtual bool OnGetContextName(FString& OutContextName) const override;
-	SCENESTATEGAMEPLAY_API virtual bool OnGetContextObject(UObject*& OutContextObject) const override;
+	UE_API virtual bool OnGetContextName(FString& OutContextName) const override;
+	UE_API virtual bool OnGetContextObject(UObject*& OutContextObject) const override;
 	//~ End USceneStatePlayer
 };
+
+#undef UE_API
diff --git a/Engine/Plugins/VirtualProduction/SceneState/Source/SceneStateMachineEditor/Public/SceneStateMachineCompiler.h b/Engine/Plugins/VirtualProduction/SceneState/Source/SceneStateMachineEditor/Public/SceneStateMachineCompiler.h
--- a/Engine/Plugins/VirtualProduction/SceneState/Source/SceneStateMachineEditor/Public/SceneStateMachineCompiler.h
+++ b/Engine/Plugins/VirtualProduction/SceneState/Source/SceneStateMachineEditor/Public/SceneStateMachineCompiler.h
@@ -11,6 +11,8 @@
 #include "UObject/ObjectKey.h"
 #include "UObject/ObjectMacros.h"
 
+#define UE_API SCENESTATEMACHINEEDITOR_API
+
 class UBlueprint;
 class UObject;
 class USceneStateMachineConduitNode;
@@ -32,9 +34,9 @@ namespace UE::SceneState::Editor
 class FStateMachineCompiler
 {
 public:
-	SCENESTATEMACHINEEDITOR_API explicit FStateMachineCompiler(USceneStateMachineGraph* InGraph, IStateMachineCompilerContext& InCompilerContext);
+	UE_API explicit FStateMachineCompiler(USceneStateMachineGraph* InGraph, IStateMachineCompilerContext& InCompilerContext);
 
-	SCENESTATEMACHINEEDITOR_API FSceneStateMachine Compile();
+	UE_API FSceneStateMachine Compile();
 
 private:
 	/** Compiles a single state (compiles state's tasks, substate machines) */
@@ -110,3 +112,5 @@ private:
 };
 
 } // UE::SceneState::Editor
+
+#undef UE_API
diff --git a/Engine/Plugins/VirtualProduction/SceneState/Source/SceneStateMachineGraph/Public/Nodes/SceneStateMachineEntryNode.h b/Engine/Plugins/VirtualProduction/SceneState/Source/SceneStateMachineGraph/Public/Nodes/SceneStateMachineEntryNode.h
--- a/Engine/Plugins/VirtualProduction/SceneState/Source/SceneStateMachineGraph/Public/Nodes/SceneStateMachineEntryNode.h
+++ b/Engine/Plugins/VirtualProduction/SceneState/Source/SceneStateMachineGraph/Public/Nodes/SceneStateMachineEntryNode.h
@@ -5,6 +5,8 @@
 #include "SceneStateMachineNode.h"
 #include "SceneStateMachineEntryNode.generated.h"
 
+#define UE_API SCENESTATEMACHINEGRAPH_API
+
 class USceneStateMachineStateNode;
 
 UCLASS(MinimalAPI, meta=(ToolTip="Entry point for state machine"))
@@ -16,7 +18,7 @@ public:
 	USceneStateMachineEntryNode();
 
 	/** Retrieves the first State Node in the State Machine */
-	SCENESTATEMACHINEGRAPH_API USceneStateMachineStateNode* GetStateNode() const;
+	UE_API USceneStateMachineStateNode* GetStateNode() const;
 
 	//~ Begin USceneStateMachineNode
 	virtual UEdGraphPin* GetInputPin() const override;
@@ -28,3 +30,5 @@ public:
 	virtual void AllocateDefaultPins() override;
 	//~ End UEdGraphNode
 };
+
+#undef UE_API
diff --git a/Engine/Plugins/VirtualProduction/SceneState/Source/SceneStateMachineGraph/Public/Nodes/SceneStateMachineNode.h b/Engine/Plugins/VirtualProduction/SceneState/Source/SceneStateMachineGraph/Public/Nodes/SceneStateMachineNode.h
--- a/Engine/Plugins/VirtualProduction/SceneState/Source/SceneStateMachineGraph/Public/Nodes/SceneStateMachineNode.h
+++ b/Engine/Plugins/VirtualProduction/SceneState/Source/SceneStateMachineGraph/Public/Nodes/SceneStateMachineNode.h
@@ -6,6 +6,8 @@
 #include "EdGraph/EdGraphNode.h"
 #include "SceneStateMachineNode.generated.h"
 
+#define UE_API SCENESTATEMACHINEGRAPH_API
+
 class USceneStateMachineTransitionNode;
 
 UCLASS(MinimalAPI, Abstract)
@@ -15,7 +17,7 @@ class USceneStateMachineNode : public UEdGraphNode
 
 public:
 	/** Gathers and returns all the transition nodes connected to this state (including bi-directional transitions) */
-    SCENESTATEMACHINEGRAPH_API TArray<USceneStateMachineTransitionNode*> GatherTransitions(bool bInSortList = false) const;
+    UE_API TArray<USceneStateMachineTransitionNode*> GatherTransitions(bool bInSortList = false) const;
 
 	FName GetNodeName() const
 	{
@@ -27,7 +29,7 @@ public:
 		return NodeType;
 	}
 
-	SCENESTATEMACHINEGRAPH_API UEdGraph* GetBoundGraph() const;
+	UE_API UEdGraph* GetBoundGraph() const;
 
 	TConstArrayView<UEdGraph*> GetBoundGraphs() const
 	{
@@ -96,3 +98,5 @@ protected:
 	UPROPERTY(Instanced)
 	TArray<TObjectPtr<UEdGraph>> BoundGraphs;
 };
+
+#undef UE_API
diff --git a/Engine/Plugins/VirtualProduction/SceneState/Source/SceneStateMachineGraph/Public/Nodes/SceneStateMachineTransitionNode.h b/Engine/Plugins/VirtualProduction/SceneState/Source/SceneStateMachineGraph/Public/Nodes/SceneStateMachineTransitionNode.h
--- a/Engine/Plugins/VirtualProduction/SceneState/Source/SceneStateMachineGraph/Public/Nodes/SceneStateMachineTransitionNode.h
+++ b/Engine/Plugins/VirtualProduction/SceneState/Source/SceneStateMachineGraph/Public/Nodes/SceneStateMachineTransitionNode.h
@@ -7,6 +7,8 @@
 #include "StructUtils/PropertyBag.h"
 #include "SceneStateMachineTransitionNode.generated.h"
 
+#define UE_API SCENESTATEMACHINEGRAPH_API
+
 UCLASS(MinimalAPI)
 class USceneStateMachineTransitionNode : public USceneStateMachineNode, public ISceneStateTransitionGraphProvider
 {
@@ -21,13 +23,13 @@ public:
 
 	void NotifyParametersChanged();
 
-	SCENESTATEMACHINEGRAPH_API static TArray<USceneStateMachineTransitionNode*> GetTransitionsToRelink(UEdGraphPin* InSourcePin
+	UE_API static TArray<USceneStateMachineTransitionNode*> GetTransitionsToRelink(UEdGraphPin* InSourcePin
 		, UEdGraphPin* InOldTargetPin
 		, TConstArrayView<UEdGraphNode*> InSelectedGraphNodes);
 
-	SCENESTATEMACHINEGRAPH_API USceneStateMachineNode* GetSourceNode() const;
+	UE_API USceneStateMachineNode* GetSourceNode() const;
 
-	SCENESTATEMACHINEGRAPH_API USceneStateMachineNode* GetTargetNode() const;
+	UE_API USceneStateMachineNode* GetTargetNode() const;
 
 	int32 GetPriority() const
 	{
@@ -128,3 +130,5 @@ private:
 	UPROPERTY(EditAnywhere, Category = "Parameters")
 	FInstancedPropertyBag Parameters;
 };
+
+#undef UE_API
diff --git a/Engine/Plugins/VirtualProduction/SceneState/Source/SceneStateMachineGraph/Public/SceneStateMachineGraphSchema.h b/Engine/Plugins/VirtualProduction/SceneState/Source/SceneStateMachineGraph/Public/SceneStateMachineGraphSchema.h
--- a/Engine/Plugins/VirtualProduction/SceneState/Source/SceneStateMachineGraph/Public/SceneStateMachineGraphSchema.h
+++ b/Engine/Plugins/VirtualProduction/SceneState/Source/SceneStateMachineGraph/Public/SceneStateMachineGraphSchema.h
@@ -5,6 +5,8 @@
 #include "EdGraph/EdGraphSchema.h"
 #include "SceneStateMachineGraphSchema.generated.h"
 
+#define UE_API SCENESTATEMACHINEGRAPH_API
+
 class USceneStateMachineNode;
 class USceneStateMachineStateNode;
 
@@ -15,20 +17,20 @@ class USceneStateMachineGraphSchema : public UEdGraphSchema
 
 public:
 	/** Pin Names */
-	SCENESTATEMACHINEGRAPH_API static const FName PN_In;
-	SCENESTATEMACHINEGRAPH_API static const FName PN_Out;
-	SCENESTATEMACHINEGRAPH_API static const FName PN_Task;
+	UE_API static const FName PN_In;
+	UE_API static const FName PN_Out;
+	UE_API static const FName PN_Task;
 
 	/** Pin Categories */
-	SCENESTATEMACHINEGRAPH_API static const FName PC_Transition;
-	SCENESTATEMACHINEGRAPH_API static const FName PC_Task;
+	UE_API static const FName PC_Transition;
+	UE_API static const FName PC_Task;
 
 	/** Pin Category Colors */
-	SCENESTATEMACHINEGRAPH_API static const FLinearColor PCC_Transition;
-	SCENESTATEMACHINEGRAPH_API static const FLinearColor PCC_Task;
+	UE_API static const FLinearColor PCC_Transition;
+	UE_API static const FLinearColor PCC_Task;
 
 	/** Attempts to find the State Node connected to the given Task Node */
-	SCENESTATEMACHINEGRAPH_API static USceneStateMachineStateNode* FindConnectedStateNode(const UEdGraphNode* InTaskNode);
+	UE_API static USceneStateMachineStateNode* FindConnectedStateNode(const UEdGraphNode* InTaskNode);
 
 	//~ Begin UEdGraphSchema
 	virtual EGraphType GetGraphType(const UEdGraph* InGraph) const override;
@@ -57,3 +59,5 @@ public:
 	};
 	void CreateConnectionWithTransition(const FTransitionConnectionParams& InParams) const;
 };
+
+#undef UE_API
diff --git a/Engine/Plugins/VirtualProduction/SceneState/Source/SceneStateTasks/Public/Setters/SceneStateSetterUtils.h b/Engine/Plugins/VirtualProduction/SceneState/Source/SceneStateTasks/Public/Setters/SceneStateSetterUtils.h
--- a/Engine/Plugins/VirtualProduction/SceneState/Source/SceneStateTasks/Public/Setters/SceneStateSetterUtils.h
+++ b/Engine/Plugins/VirtualProduction/SceneState/Source/SceneStateTasks/Public/Setters/SceneStateSetterUtils.h
@@ -7,6 +7,8 @@
 #include "SceneStateExecutionContext.h"
 #include "SceneStatePropertyUtils.h"
 
+#define UE_API SCENESTATETASKS_API
+
 struct FSceneStatePropertyReference;
 
 namespace UE::SceneState
@@ -19,7 +21,7 @@ namespace UE::SceneState
 	 * @param InArrayIndex the array index to set within the destination property (necessary if property has an Array Dim > 1, 0 otherwise)
 	 * @return true if the setter call was made
 	 */
-	SCENESTATETASKS_API bool CallValueSetter(uint8* InDestValuePtr, const FSceneStateBindingResolvedReference& InPropertyReference, const uint8* InSourceValuePtr, int32 InArrayIndex = 0);
+	UE_API bool CallValueSetter(uint8* InDestValuePtr, const FSceneStateBindingResolvedReference& InPropertyReference, const uint8* InSourceValuePtr, int32 InArrayIndex = 0);
 
 	/** A Setter Task Instance is valid if it has a valid 'Value' member, as well as a 'Target' FSceneStatePropertyReference member */
 	template<typename InTaskInstanceType>
@@ -75,3 +77,5 @@ namespace UE::SceneState
 	}
 
 } // UE::SceneState
+
+#undef UE_API
