Motion Design Scene State: added schema to set the rules of which tasks are allowed.

#Synced-CL 50160130

diff --git a/Engine/Plugins/VirtualProduction/AvalancheSceneState/Source/AvalancheSceneState/Private/RemoteControl/AvaSceneStateRCTask.cpp b/Engine/Plugins/VirtualProduction/AvalancheSceneState/Source/AvalancheSceneState/Private/RemoteControl/AvaSceneStateRCTask.cpp
--- a/Engine/Plugins/VirtualProduction/AvalancheSceneState/Source/AvalancheSceneState/Private/RemoteControl/AvaSceneStateRCTask.cpp
+++ b/Engine/Plugins/VirtualProduction/AvalancheSceneState/Source/AvalancheSceneState/Private/RemoteControl/AvaSceneStateRCTask.cpp
@@ -2,6 +2,7 @@
 
 #include "RemoteControl/AvaSceneStateRCTask.h"
 #include "AvaSceneStateLog.h"
+#include "AvaSceneStateSchema.h"
 #include "AvaSceneStateUtils.h"
 #include "Controller/RCController.h"
 #include "IAvaSceneInterface.h"
@@ -89,6 +90,11 @@ const UScriptStruct* FAvaSceneStateRCTask::OnGetTaskInstanceType() const
 	return FInstanceDataType::StaticStruct();
 }
 
+bool FAvaSceneStateRCTask::OnSupportsSchema(TNotNull<const USceneStateSchema*> InSchema) const
+{
+	return Super::OnSupportsSchema(InSchema) && InSchema->IsA<UAvaSceneStateSchema>();
+}
+
 void FAvaSceneStateRCTask::OnBuildTaskInstance(UObject* InOuter, FStructView InTaskInstance) const
 {
 	FInstanceDataType& Instance = InTaskInstance.Get<FInstanceDataType>();
diff --git a/Engine/Plugins/VirtualProduction/AvalancheSceneState/Source/AvalancheSceneState/Private/Sequence/AvaSceneStatePlaySequenceTask.cpp b/Engine/Plugins/VirtualProduction/AvalancheSceneState/Source/AvalancheSceneState/Private/Sequence/AvaSceneStatePlaySequenceTask.cpp
--- a/Engine/Plugins/VirtualProduction/AvalancheSceneState/Source/AvalancheSceneState/Private/Sequence/AvaSceneStatePlaySequenceTask.cpp
+++ b/Engine/Plugins/VirtualProduction/AvalancheSceneState/Source/AvalancheSceneState/Private/Sequence/AvaSceneStatePlaySequenceTask.cpp
@@ -2,6 +2,7 @@
 
 #include "Sequence/AvaSceneStatePlaySequenceTask.h"
 #include "AvaSceneStateLog.h"
+#include "AvaSceneStateSchema.h"
 #include "AvaSceneStateUtils.h"
 #include "AvaSequence.h"
 #include "AvaSequencePlayer.h"
@@ -15,6 +16,11 @@ const UScriptStruct* FAvaSceneStatePlaySequenceTask::OnGetTaskInstanceType() con
 {
 	return FInstanceDataType::StaticStruct();
 }
+
+bool FAvaSceneStatePlaySequenceTask::OnSupportsSchema(TNotNull<const USceneStateSchema*> InSchema) const
+{
+	return Super::OnSupportsSchema(InSchema) && InSchema->IsA<UAvaSceneStateSchema>();
+}
 #endif
 
 void FAvaSceneStatePlaySequenceTask::OnStart(const FSceneStateExecutionContext& InContext, FStructView InTaskInstance) const
diff --git a/Engine/Plugins/VirtualProduction/AvalancheSceneState/Source/AvalancheSceneState/Private/Ticker/AvaSceneStateSpawnerTickerTask.cpp b/Engine/Plugins/VirtualProduction/AvalancheSceneState/Source/AvalancheSceneState/Private/Ticker/AvaSceneStateSpawnerTickerTask.cpp
--- a/Engine/Plugins/VirtualProduction/AvalancheSceneState/Source/AvalancheSceneState/Private/Ticker/AvaSceneStateSpawnerTickerTask.cpp
+++ b/Engine/Plugins/VirtualProduction/AvalancheSceneState/Source/AvalancheSceneState/Private/Ticker/AvaSceneStateSpawnerTickerTask.cpp
@@ -1,6 +1,7 @@
 // Copyright Epic Games, Inc. All Rights Reserved.
 
 #include "Ticker/AvaSceneStateSpawnerTickerTask.h"
+#include "AvaSceneStateSchema.h"
 #include "Framework/Ticker/AvaTickerComponent.h"
 #include "GameFramework/Actor.h"
 
@@ -11,6 +12,11 @@ const UScriptStruct* FAvaSceneStateSpawnerTickerTask::OnGetTaskInstanceType() co
 {
 	return FInstanceDataType::StaticStruct();
 }
+
+bool FAvaSceneStateSpawnerTickerTask::OnSupportsSchema(TNotNull<const USceneStateSchema*> InSchema) const
+{
+	return Super::OnSupportsSchema(InSchema) && InSchema->IsA<UAvaSceneStateSchema>();
+}
 #endif
 
 bool FAvaSceneStateSpawnerTickerTask::ShouldSpawnActor(FStructView InTaskInstance, FText& OutErrorMessage) const
diff --git a/Engine/Plugins/VirtualProduction/AvalancheSceneState/Source/AvalancheSceneState/Public/AvaSceneStateSchema.h b/Engine/Plugins/VirtualProduction/AvalancheSceneState/Source/AvalancheSceneState/Public/AvaSceneStateSchema.h
new file mode 100644
--- /dev/null
+++ b/Engine/Plugins/VirtualProduction/AvalancheSceneState/Source/AvalancheSceneState/Public/AvaSceneStateSchema.h
@@ -0,0 +1,13 @@
+// Copyright Epic Games, Inc. All Rights Reserved.
+
+#pragma once
+
+#include "SceneStateGameplaySchema.h"
+#include "AvaSceneStateSchema.generated.h"
+
+/** Schema for Motion Design scene states. */
+UCLASS(MinimalAPI, DisplayName="Motion Design Scene State Schema")
+class UAvaSceneStateSchema : public USceneStateGameplaySchema
+{
+	GENERATED_BODY()
+};
diff --git a/Engine/Plugins/VirtualProduction/AvalancheSceneState/Source/AvalancheSceneState/Public/RemoteControl/AvaSceneStateRCTask.h b/Engine/Plugins/VirtualProduction/AvalancheSceneState/Source/AvalancheSceneState/Public/RemoteControl/AvaSceneStateRCTask.h
--- a/Engine/Plugins/VirtualProduction/AvalancheSceneState/Source/AvalancheSceneState/Public/RemoteControl/AvaSceneStateRCTask.h
+++ b/Engine/Plugins/VirtualProduction/AvalancheSceneState/Source/AvalancheSceneState/Public/RemoteControl/AvaSceneStateRCTask.h
@@ -54,7 +54,7 @@ struct FAvaSceneStateRCTaskInstance : public FSceneStateTaskInstance
 	FInstancedPropertyBag ControllerValues;
 };
 
-USTRUCT(DisplayName="Set RC Controller Values", Category="Motion Design")
+USTRUCT(DisplayName="Set RC Controller Values", Category="Motion Design", meta=(RequiresContextWorld, WithSupportsSchema))
 struct FAvaSceneStateRCTask : public FSceneStateTask
 {
 	GENERATED_BODY()
@@ -67,6 +67,7 @@ protected:
 	//~ Begin FSceneStateTask
 #if WITH_EDITOR
 	virtual const UScriptStruct* OnGetTaskInstanceType() const override;
+	virtual bool OnSupportsSchema(TNotNull<const USceneStateSchema*> InSchema) const override;
 	virtual void OnBuildTaskInstance(UObject* InOuter, FStructView InTaskInstance) const override;
 #endif
 	virtual const FSceneStateTaskBindingExtension* OnGetBindingExtension() const override;
diff --git a/Engine/Plugins/VirtualProduction/AvalancheSceneState/Source/AvalancheSceneState/Public/Sequence/AvaSceneStatePlaySequenceTask.h b/Engine/Plugins/VirtualProduction/AvalancheSceneState/Source/AvalancheSceneState/Public/Sequence/AvaSceneStatePlaySequenceTask.h
--- a/Engine/Plugins/VirtualProduction/AvalancheSceneState/Source/AvalancheSceneState/Public/Sequence/AvaSceneStatePlaySequenceTask.h
+++ b/Engine/Plugins/VirtualProduction/AvalancheSceneState/Source/AvalancheSceneState/Public/Sequence/AvaSceneStatePlaySequenceTask.h
@@ -49,7 +49,7 @@ struct FAvaSceneStatePlaySequenceTaskInstance : public FSceneStateTaskInstance
 	FDelegateHandle OnSequenceFinishedHandle;
 };
 
-USTRUCT(DisplayName="Play Sequence", Category="Motion Design", meta=(ToolTip="Plays a Motion Design Sequence"))
+USTRUCT(DisplayName="Play Sequence", Category="Motion Design", meta=(RequiresContextWorld, WithSupportsSchema, ToolTip="Plays a Motion Design Sequence"))
 struct FAvaSceneStatePlaySequenceTask : public FSceneStateTask
 {
 	GENERATED_BODY()
@@ -60,6 +60,7 @@ protected:
 	//~ Begin FSceneStateTask
 #if WITH_EDITOR
 	virtual const UScriptStruct* OnGetTaskInstanceType() const override;
+	virtual bool OnSupportsSchema(TNotNull<const USceneStateSchema*> InSchema) const override;
 #endif
 	virtual void OnStart(const FSceneStateExecutionContext& InContext, FStructView InTaskInstance) const override;
 	virtual void OnStop(const FSceneStateExecutionContext& InContext, FStructView InTaskInstance, ESceneStateTaskStopReason InStopReason) const override;
diff --git a/Engine/Plugins/VirtualProduction/AvalancheSceneState/Source/AvalancheSceneState/Public/Ticker/AvaSceneStateSpawnerTickerTask.h b/Engine/Plugins/VirtualProduction/AvalancheSceneState/Source/AvalancheSceneState/Public/Ticker/AvaSceneStateSpawnerTickerTask.h
--- a/Engine/Plugins/VirtualProduction/AvalancheSceneState/Source/AvalancheSceneState/Public/Ticker/AvaSceneStateSpawnerTickerTask.h
+++ b/Engine/Plugins/VirtualProduction/AvalancheSceneState/Source/AvalancheSceneState/Public/Ticker/AvaSceneStateSpawnerTickerTask.h
@@ -21,7 +21,7 @@ struct FAvaSceneStateTickerTaskInstance : public FSceneStateSpawnActorTaskInstan
  * Spawns an actor and adds it to the Ticker.
  * Note: if the Ticker actor is invalid or does not contain a Ticker Component, the actor will not spawn.
  */
-USTRUCT(DisplayName="Spawn Actor to Ticker", Category="Motion Design")
+USTRUCT(DisplayName="Spawn Actor to Ticker", Category="Motion Design", meta=(WithSupportsSchema))
 struct FAvaSceneStateSpawnerTickerTask : public FSceneStateSpawnActorTask
 {
 	GENERATED_BODY()
@@ -32,6 +32,7 @@ protected:
 	//~ Begin FSceneStateTask
 #if WITH_EDITOR
 	UE_API virtual const UScriptStruct* OnGetTaskInstanceType() const override;
+	UE_API virtual bool OnSupportsSchema(TNotNull<const USceneStateSchema*> InSchema) const override;
 #endif
 	//~ End FSceneStateTask
 
diff --git a/Engine/Plugins/VirtualProduction/AvalancheSceneState/Source/AvalancheSceneStateBlueprint/AvalancheSceneStateBlueprint.Build.cs b/Engine/Plugins/VirtualProduction/AvalancheSceneState/Source/AvalancheSceneStateBlueprint/AvalancheSceneStateBlueprint.Build.cs
--- a/Engine/Plugins/VirtualProduction/AvalancheSceneState/Source/AvalancheSceneStateBlueprint/AvalancheSceneStateBlueprint.Build.cs
+++ b/Engine/Plugins/VirtualProduction/AvalancheSceneState/Source/AvalancheSceneStateBlueprint/AvalancheSceneStateBlueprint.Build.cs
@@ -23,6 +23,7 @@ public class AvalancheSceneStateBlueprint : ModuleRules
             {
                 "AvalancheSceneState",
                 "CoreUObject",
+                "SceneState",
             }
         );
     }
diff --git a/Engine/Plugins/VirtualProduction/AvalancheSceneState/Source/AvalancheSceneStateBlueprint/Private/AvaSceneStateBlueprint.cpp b/Engine/Plugins/VirtualProduction/AvalancheSceneState/Source/AvalancheSceneStateBlueprint/Private/AvaSceneStateBlueprint.cpp
new file mode 100644
--- /dev/null
+++ b/Engine/Plugins/VirtualProduction/AvalancheSceneState/Source/AvalancheSceneStateBlueprint/Private/AvaSceneStateBlueprint.cpp
@@ -0,0 +1,16 @@
+// Copyright Epic Games, Inc. All Rights Reserved.
+
+#include "AvaSceneStateBlueprint.h"
+#include "AvaSceneStateSchema.h"
+#include "SceneStateVersion.h"
+
+void UAvaSceneStateBlueprint::PostLoad()
+{
+	Super::PostLoad();
+
+	// Previous versions to default to the motion design schema type
+	if (GetLinkerCustomVersion(UE::SceneState::FVersion::Guid) < UE::SceneState::FVersion::Schema)
+	{
+		SetSceneStateSchema(UAvaSceneStateSchema::StaticClass());
+	}
+}
diff --git a/Engine/Plugins/VirtualProduction/AvalancheSceneState/Source/AvalancheSceneStateBlueprint/Public/AvaSceneStateBlueprint.h b/Engine/Plugins/VirtualProduction/AvalancheSceneState/Source/AvalancheSceneStateBlueprint/Public/AvaSceneStateBlueprint.h
--- a/Engine/Plugins/VirtualProduction/AvalancheSceneState/Source/AvalancheSceneStateBlueprint/Public/AvaSceneStateBlueprint.h
+++ b/Engine/Plugins/VirtualProduction/AvalancheSceneState/Source/AvalancheSceneStateBlueprint/Public/AvaSceneStateBlueprint.h
@@ -5,12 +5,17 @@
 #include "SceneStateBlueprint.h"
 #include "AvaSceneStateBlueprint.generated.h"
 
-/**
- * Scene State Blueprint child that overrides the Property Binder class with a child class
- * @see UAvaSceneStateBinder
- */
+#define UE_API AVALANCHESCENESTATEBLUEPRINT_API
+
 UCLASS(MinimalAPI, DisplayName="Motion Design Scene State Blueprint")
 class UAvaSceneStateBlueprint : public USceneStateBlueprint
 {
 	GENERATED_BODY()
+
+public:
+	//~ Begin UObject
+	UE_API virtual void PostLoad() override;
+	//~ End UObject
 };
+
+#undef UE_API
diff --git a/Engine/Plugins/VirtualProduction/AvalancheSceneState/Source/AvalancheSceneStateEditor/Private/AvaSceneStateExtension.cpp b/Engine/Plugins/VirtualProduction/AvalancheSceneState/Source/AvalancheSceneStateEditor/Private/AvaSceneStateExtension.cpp
--- a/Engine/Plugins/VirtualProduction/AvalancheSceneState/Source/AvalancheSceneStateEditor/Private/AvaSceneStateExtension.cpp
+++ b/Engine/Plugins/VirtualProduction/AvalancheSceneState/Source/AvalancheSceneStateEditor/Private/AvaSceneStateExtension.cpp
@@ -3,6 +3,7 @@
 #include "AvaSceneStateExtension.h"
 #include "AvaSceneStateActor.h"
 #include "AvaSceneStateBlueprint.h"
+#include "AvaSceneStateSchema.h"
 #include "BlueprintActionDatabase.h"
 #include "Editor.h"
 #include "Editor/UnrealEdEngine.h"
@@ -117,10 +118,11 @@ USceneStateBlueprint* FAvaSceneStateExtension::CreateSceneStateBlueprint(AAvaSce
 {
 	check(InSceneStateActor);
 
-	UFactory* SceneStateBlueprintFactory = NewObject<USceneStateBlueprintFactory>(GetTransientPackage());
-	check(SceneStateBlueprintFactory);
+	USceneStateBlueprintFactory* Factory = NewObject<USceneStateBlueprintFactory>(GetTransientPackage());
+	check(Factory);
+	Factory->SchemaClass = UAvaSceneStateSchema::StaticClass();
 
-	USceneStateBlueprint* NewSceneStateBlueprint = CastChecked<USceneStateBlueprint>(SceneStateBlueprintFactory->FactoryCreateNew(UAvaSceneStateBlueprint::StaticClass()
+	USceneStateBlueprint* NewSceneStateBlueprint = CastChecked<USceneStateBlueprint>(Factory->FactoryCreateNew(UAvaSceneStateBlueprint::StaticClass()
 		, InSceneStateActor
 		, TEXT("SceneStateBlueprint")
 		, RF_Transactional
diff --git a/Engine/Plugins/VirtualProduction/SceneState/Source/SceneState/Private/SceneStateSchema.cpp b/Engine/Plugins/VirtualProduction/SceneState/Source/SceneState/Private/SceneStateSchema.cpp
new file mode 100644
--- /dev/null
+++ b/Engine/Plugins/VirtualProduction/SceneState/Source/SceneState/Private/SceneStateSchema.cpp
@@ -0,0 +1,23 @@
+// Copyright Epic Games, Inc. All Rights Reserved.
+
+#include "SceneStateSchema.h"
+#include "Tasks/SceneStateTask.h"
+
+bool USceneStateSchema::IsTaskStructAllowed(TSubScriptStructOf<FSceneStateTask> InTaskStruct) const
+{
+	return OnIsTaskStructAllowed(InTaskStruct);
+}
+
+bool USceneStateSchema::IsTaskAssetAllowed(const FAssetData& InTaskAsset) const
+{
+	return OnIsTaskAssetAllowed(InTaskAsset);
+}
+
+TSubclassOf<UObject> USceneStateSchema::GetContextObjectClass() const
+{
+	if (!ContextObjectClass)
+	{
+		ContextObjectClass = OnGetContextObjectClass();
+	}
+	return ContextObjectClass;
+}
diff --git a/Engine/Plugins/VirtualProduction/SceneState/Source/SceneState/Private/SceneStateVersion.cpp b/Engine/Plugins/VirtualProduction/SceneState/Source/SceneState/Private/SceneStateVersion.cpp
new file mode 100644
--- /dev/null
+++ b/Engine/Plugins/VirtualProduction/SceneState/Source/SceneState/Private/SceneStateVersion.cpp
@@ -0,0 +1,11 @@
+// Copyright Epic Games, Inc. All Rights Reserved.
+
+#include "SceneStateVersion.h"
+#include "Serialization/CustomVersion.h"
+
+namespace UE::SceneState
+{
+	const FGuid FVersion::Guid(0x2333E4E3, 0xA940435F, 0x8B481930, 0x906B27F9);
+	FCustomVersionRegistration GRegisterSceneStateVersion(FVersion::Guid, FVersion::LatestVersion, TEXT("SceneStateVersion"));
+
+} // UE::SceneState
diff --git a/Engine/Plugins/VirtualProduction/SceneState/Source/SceneState/Private/Tasks/SceneStateTask.cpp b/Engine/Plugins/VirtualProduction/SceneState/Source/SceneState/Private/Tasks/SceneStateTask.cpp
--- a/Engine/Plugins/VirtualProduction/SceneState/Source/SceneState/Private/Tasks/SceneStateTask.cpp
+++ b/Engine/Plugins/VirtualProduction/SceneState/Source/SceneState/Private/Tasks/SceneStateTask.cpp
@@ -33,6 +33,11 @@ const UScriptStruct* FSceneStateTask::GetTaskInstanceType() const
 	return OnGetTaskInstanceType();
 }
 
+bool FSceneStateTask::SupportsSchema(const USceneStateSchema* InSchema) const
+{
+	return InSchema && OnSupportsSchema(InSchema);
+}
+
 void FSceneStateTask::BuildTaskInstance(UObject* InOuter, FStructView InTaskInstance) const
 {
 	if (InTaskInstance.IsValid())
diff --git a/Engine/Plugins/VirtualProduction/SceneState/Source/SceneState/Public/ISceneStateSchemaProvider.h b/Engine/Plugins/VirtualProduction/SceneState/Source/SceneState/Public/ISceneStateSchemaProvider.h
new file mode 100644
--- /dev/null
+++ b/Engine/Plugins/VirtualProduction/SceneState/Source/SceneState/Public/ISceneStateSchemaProvider.h
@@ -0,0 +1,23 @@
+// Copyright Epic Games, Inc. All Rights Reserved.
+
+#pragma once
+
+#include "UObject/Interface.h"
+#include "ISceneStateSchemaProvider.generated.h"
+
+class USceneStateSchema;
+
+UINTERFACE(MinimalAPI, meta=(CannotImplementInterfaceInBlueprint))
+class USceneStateSchemaProvider : public UInterface
+{
+	GENERATED_BODY()
+};
+
+class ISceneStateSchemaProvider
+{
+	GENERATED_BODY()
+
+public:
+	/** Returns the schema used for the scene state */
+	virtual TSubclassOf<USceneStateSchema> GetSceneStateSchemaClass() const = 0;
+};
diff --git a/Engine/Plugins/VirtualProduction/SceneState/Source/SceneState/Public/SceneStateSchema.h b/Engine/Plugins/VirtualProduction/SceneState/Source/SceneState/Public/SceneStateSchema.h
new file mode 100644
--- /dev/null
+++ b/Engine/Plugins/VirtualProduction/SceneState/Source/SceneState/Public/SceneStateSchema.h
@@ -0,0 +1,55 @@
+// Copyright Epic Games, Inc. All Rights Reserved.
+
+#pragma once
+
+#include "Templates/SubScriptStructOf.h"
+#include "Templates/SubclassOf.h"
+#include "UObject/Object.h"
+#include "SceneStateSchema.generated.h"
+
+#define UE_API SCENESTATE_API
+
+struct FAssetData;
+struct FSceneStateTask;
+
+/** Determines the rules (e.g. allowed tasks) and provides context for a Scene State object type. */
+UCLASS(MinimalAPI, Abstract)
+class USceneStateSchema : public UObject
+{
+	GENERATED_BODY()
+
+public:
+	/** Determines whether the given task struct is allowed */
+	UE_API bool IsTaskStructAllowed(TSubScriptStructOf<FSceneStateTask> InTaskStruct) const;
+
+	/** Determines whether the given task asset (task blueprint) is allowed */
+	UE_API bool IsTaskAssetAllowed(const FAssetData& InTaskAsset) const;
+
+	/** Return the context object class supported */
+	UE_API TSubclassOf<UObject> GetContextObjectClass() const;
+
+protected:
+	/** Determines whether a given task struct is allowed */
+	virtual bool OnIsTaskStructAllowed(TSubScriptStructOf<FSceneStateTask> InTaskStruct) const
+	{
+		return true;
+	}
+
+	/** Determines whether the given task asset (task blueprint) is allowed */
+	virtual bool OnIsTaskAssetAllowed(const FAssetData& InTaskAsset) const
+	{
+		return true;
+	}
+
+	/** Return the context object class supported. Called only once to keep the context object class constant throughout the schema. */
+	virtual TSubclassOf<UObject> OnGetContextObjectClass() const
+	{
+		return TSubclassOf<UObject>();
+	}
+
+private:
+	UPROPERTY()
+	mutable TSubclassOf<UObject> ContextObjectClass;
+};
+
+#undef UE_API
diff --git a/Engine/Plugins/VirtualProduction/SceneState/Source/SceneState/Public/SceneStateVersion.h b/Engine/Plugins/VirtualProduction/SceneState/Source/SceneState/Public/SceneStateVersion.h
new file mode 100644
--- /dev/null
+++ b/Engine/Plugins/VirtualProduction/SceneState/Source/SceneState/Public/SceneStateVersion.h
@@ -0,0 +1,36 @@
+// Copyright Epic Games, Inc. All Rights Reserved.
+
+#pragma once
+
+#include "Misc/Guid.h"
+
+#define UE_API SCENESTATE_API
+
+namespace UE::SceneState
+{
+
+/** Describes the versions of scene state */
+struct FVersion
+{
+private:
+	FVersion() = delete;
+
+public:
+	enum Type : uint8
+	{
+		PreVersioning = 0,
+
+		/** Added schemas to set rules for scene state blueprints */
+		Schema,
+
+		/* ------------------------------------------------------ */
+		VersionPlusOne,
+		LatestVersion = VersionPlusOne - 1
+	};
+
+	UE_API const static FGuid Guid;
+};
+
+} // UE::SceneState
+
+#undef UE_API
diff --git a/Engine/Plugins/VirtualProduction/SceneState/Source/SceneState/Public/Tasks/SceneStateTask.h b/Engine/Plugins/VirtualProduction/SceneState/Source/SceneState/Public/Tasks/SceneStateTask.h
--- a/Engine/Plugins/VirtualProduction/SceneState/Source/SceneState/Public/Tasks/SceneStateTask.h
+++ b/Engine/Plugins/VirtualProduction/SceneState/Source/SceneState/Public/Tasks/SceneStateTask.h
@@ -9,6 +9,7 @@
 
 #define UE_API SCENESTATE_API
 
+class USceneStateSchema;
 struct FSceneStateBindingDesc;
 struct FSceneStateExecutionContext;
 struct FSceneStateTaskBindingExtension;
@@ -32,6 +33,7 @@ namespace UE::SceneState
  * Tasks are immutable in execution time, and so are meant to only hold logic and template read-only data.
  * Each Task has a Task Instance Type that is allocated and used for instance data that the task can then mutate.
  * @see FSceneStateTaskInstance
+ * @see SceneStateTaskMetadata.h for information about metadata that can be used with Tasks
  */
 USTRUCT(meta=(Hidden))
 struct FSceneStateTask
@@ -47,6 +49,12 @@ struct FSceneStateTask
 	 */
 	UE_API const UScriptStruct* GetTaskInstanceType() const;
 
+	/**
+	 * Determines whether the given schema is supported by this task 
+	 * This is called only when the task struct contains the 'WithSupportsSchema' metadata.
+	 */
+	UE_API bool SupportsSchema(const USceneStateSchema* InSchema) const;
+
 	/**
 	 * Called to init properties (e.g. Guids) or instances objects within the task instance
 	 * @param InOuter the outer to use for instanced objects
@@ -126,6 +134,15 @@ protected:
 		return nullptr;
 	}
 
+	/**
+	 * Determines whether the given schema is supported by this task 
+	 * This is called only when the task struct contains the 'WithSupportsSchema' metadata.
+	 */
+	virtual bool OnSupportsSchema(TNotNull<const USceneStateSchema*> InSchema) const
+	{
+		return true;
+	}
+
 	/**
 	 * Called to init properties (e.g. Guids) or instances objects within the task instance
 	 * @see BuildTaskInstance
@@ -198,7 +215,7 @@ private:
 	UPROPERTY()
 	FSceneStateRange PrerequisiteRange;
 
-	/** Additional Information about how this Task (e.g. how it Ticks, etc) */
+	/** Additional information about this task (e.g. if it ticks, whether it extends bindings, etc.) */
 	ESceneStateTaskFlags TaskFlags = ESceneStateTaskFlags::None;
 
 	friend UE::SceneState::Editor::FBindingCompiler;
diff --git a/Engine/Plugins/VirtualProduction/SceneState/Source/SceneState/Public/Tasks/SceneStateTaskMetadata.h b/Engine/Plugins/VirtualProduction/SceneState/Source/SceneState/Public/Tasks/SceneStateTaskMetadata.h
--- a/Engine/Plugins/VirtualProduction/SceneState/Source/SceneState/Public/Tasks/SceneStateTaskMetadata.h
+++ b/Engine/Plugins/VirtualProduction/SceneState/Source/SceneState/Public/Tasks/SceneStateTaskMetadata.h
@@ -4,6 +4,23 @@
 
 #include "SceneStateTaskMetadata.generated.h"
 
+namespace UE::SceneState::Editor
+{
+#if WITH_EDITOR
+	/** The task should not appear in the list of available tasks (e.g. because it's a base task struct) */
+	const FLazyName MD_Hidden(TEXT("Hidden"));
+
+	/**
+	 * The task requires the context object to implement GetWorld().
+	 * This is just a convenience metadata to avoid having to call SupportsSchema.
+	 */
+	const FLazyName MD_RequiresContextWorld(TEXT("RequiresContextWorld"));
+
+	/** The task implements OnSupportsSchema function. Tasks without this metadata will not have this editor function called. */
+	const FLazyName MD_WithSupportsSchema(TEXT("WithSupportsSchema"));
+#endif
+}
+
 /** Metadata information about the Task. Available only in editor */
 USTRUCT()
 struct FSceneStateTaskMetadata
diff --git a/Engine/Plugins/VirtualProduction/SceneState/Source/SceneStateBlueprint/Private/SceneStateBlueprint.cpp b/Engine/Plugins/VirtualProduction/SceneState/Source/SceneStateBlueprint/Private/SceneStateBlueprint.cpp
--- a/Engine/Plugins/VirtualProduction/SceneState/Source/SceneStateBlueprint/Private/SceneStateBlueprint.cpp
+++ b/Engine/Plugins/VirtualProduction/SceneState/Source/SceneStateBlueprint/Private/SceneStateBlueprint.cpp
@@ -19,6 +19,8 @@
 #include "SceneStateBlueprintUtils.h"
 #include "SceneStateGeneratedClass.h"
 #include "SceneStateMachineGraph.h"
+#include "SceneStateVersion.h"
+#include "Tasks/SceneStateTaskMetadata.h"
 #include "UObject/UObjectIterator.h"
 
 #define LOCTEXT_NAMESPACE "SceneStateBlueprint"
@@ -150,7 +152,6 @@ const FSceneStateBindingCollection& USceneStateBlueprint::GetBindingCollection()
 
 bool USceneStateBlueprint::ForEachBindableFunction(TFunctionRef<bool(const FSceneStateBindingDesc&, const UE::SceneState::FBindingFunctionInfo&)> InFunc) const
 {
-	const FName MD_Hidden = TEXT("Hidden");
 	const FName MD_Category = TEXT("Category");
 
 	const FText Section = LOCTEXT("BindingFunctionsSectionLabel", "Functions");
@@ -158,7 +159,7 @@ bool USceneStateBlueprint::ForEachBindableFunction(TFunctionRef<bool(const FScen
 	// Todo: cache functions for faster iteration
 	for (const UScriptStruct* Struct : TObjectRange<UScriptStruct>())
 	{
-		if (!Struct->IsChildOf<FSceneStateFunction>() || Struct->HasMetaData(MD_Hidden))
+		if (!Struct->IsChildOf<FSceneStateFunction>() || Struct->HasMetaData(UE::SceneState::Editor::MD_Hidden))
 		{
 			continue;
 		}
@@ -262,6 +263,16 @@ void USceneStateBlueprint::AppendBindablePropertyFunctionStructs(TArray<TInstanc
 		});
 }
 
+void USceneStateBlueprint::SetSceneStateSchema(TSubclassOf<USceneStateSchema> InSchemaClass)
+{
+	SceneStateSchemaClass = InSchemaClass;
+}
+
+TSubclassOf<USceneStateSchema> USceneStateBlueprint::GetSceneStateSchemaClass() const
+{
+	return SceneStateSchemaClass;
+}
+
 void USceneStateBlueprint::SetObjectBeingDebugged(UObject* InNewObject)
 {
 	using namespace UE::SceneState;
@@ -303,6 +314,31 @@ bool USceneStateBlueprint::IsValidForBytecodeOnlyRecompile() const
 	return false;
 }
 
+void USceneStateBlueprint::Serialize(FArchive& Ar)
+{
+	Super::Serialize(Ar);
+	Ar.UsingCustomVersion(UE::SceneState::FVersion::Guid);
+}
+
+void USceneStateBlueprint::PostLoad()
+{
+	Super::PostLoad();
+
+	// Default previous versions to the motion design schema.
+	// If motion design isn't enabled, fall back to the gameplay schema.
+	if (GetLinkerCustomVersion(UE::SceneState::FVersion::Guid) < UE::SceneState::FVersion::Schema)
+	{
+		if (UClass* const MotionDesignSchemaClass = FindObject<UClass>(nullptr, TEXT("/Script/AvalancheSceneState.AvaSceneStateSchema")))
+		{
+			SetSceneStateSchema(MotionDesignSchemaClass);
+		}
+		else if (UClass* const GameplaySchemaClass = FindObject<UClass>(nullptr, TEXT("/Script/SceneStateGameplay.SceneStateGameplaySchema")))
+		{
+			SetSceneStateSchema(GameplaySchemaClass);
+		}
+	}
+}
+
 void USceneStateBlueprint::BeginDestroy()
 {
 	Super::BeginDestroy();
diff --git a/Engine/Plugins/VirtualProduction/SceneState/Source/SceneStateBlueprint/Public/SceneStateBlueprint.h b/Engine/Plugins/VirtualProduction/SceneState/Source/SceneStateBlueprint/Public/SceneStateBlueprint.h
--- a/Engine/Plugins/VirtualProduction/SceneState/Source/SceneStateBlueprint/Public/SceneStateBlueprint.h
+++ b/Engine/Plugins/VirtualProduction/SceneState/Source/SceneStateBlueprint/Public/SceneStateBlueprint.h
@@ -4,6 +4,7 @@
 
 #include "Containers/Array.h"
 #include "Engine/Blueprint.h"
+#include "ISceneStateSchemaProvider.h"
 #include "SceneStateBindingCollection.h"
 #include "SceneStateBindingCollectionOwner.h"
 #include "SceneStateBlueprint.generated.h"
@@ -14,7 +15,7 @@ class USceneStateMachineGraph;
 class USceneStateMachineStateNode;
 
 UCLASS(MinimalAPI, DisplayName="Motion Design Scene State Blueprint")
-class USceneStateBlueprint : public UBlueprint, public ISceneStateBindingCollectionOwner
+class USceneStateBlueprint : public UBlueprint, public ISceneStateSchemaProvider, public ISceneStateBindingCollectionOwner
 {
 	GENERATED_BODY()
 
@@ -68,6 +69,13 @@ public:
 	UE_API virtual void AppendBindablePropertyFunctionStructs(TArray<TInstancedStruct<FPropertyBindingBindableStructDescriptor>>& InOutStructs) const override;
 	//~ End IPropertyBindingBindingCollectionOwner
 
+	/** Sets the scene state schema to use for this blueprint */
+	UE_API void SetSceneStateSchema(TSubclassOf<USceneStateSchema> InSchemaClass);
+
+	//~ Begin ISceneStateSchemaProvider
+	UE_API virtual TSubclassOf<USceneStateSchema> GetSceneStateSchemaClass() const override;
+	//~ End ISceneStateSchemaProvider
+
 	//~ Begin UBlueprint
 	UE_API virtual void SetObjectBeingDebugged(UObject* InNewObject) override;
 	UE_API virtual UClass* GetBlueprintClass() const override;
@@ -78,6 +86,8 @@ public:
 	//~ End UBlueprint
 
 	//~ Begin UObject
+	UE_API virtual void Serialize(FArchive& Ar) override;
+	UE_API virtual void PostLoad() override;
 	UE_API virtual void BeginDestroy() override;
 	//~ End UObject
 
@@ -102,6 +112,10 @@ private:
 	/** Called when a state machine graph's parameters have changed. Used to fix the path to parameters that have possibly been renamed */
 	void UpdateGraphParametersBindings(FPropertyBindingPath& InPath, USceneStateMachineGraph* InGraph);
 
+	/** Scene State schema class to use for this blueprint */
+	UPROPERTY(VisibleAnywhere, Category="Scene State", AssetRegistrySearchable)
+	TSubclassOf<USceneStateSchema> SceneStateSchemaClass;
+
 	/** Unique id representing this blueprint / generated class as a bindable struct */
 	UPROPERTY()
 	FGuid RootId;
diff --git a/Engine/Plugins/VirtualProduction/SceneState/Source/SceneStateBlueprintEditor/Private/SceneStateBlueprintFactory.cpp b/Engine/Plugins/VirtualProduction/SceneState/Source/SceneStateBlueprintEditor/Private/SceneStateBlueprintFactory.cpp
--- a/Engine/Plugins/VirtualProduction/SceneState/Source/SceneStateBlueprintEditor/Private/SceneStateBlueprintFactory.cpp
+++ b/Engine/Plugins/VirtualProduction/SceneState/Source/SceneStateBlueprintEditor/Private/SceneStateBlueprintFactory.cpp
@@ -2,8 +2,10 @@
 
 #include "SceneStateBlueprintFactory.h"
 #include "AssetToolsModule.h"
+#include "ClassViewerFilter.h"
 #include "Kismet2/BlueprintEditorUtils.h"
 #include "Kismet2/KismetEditorUtilities.h"
+#include "Kismet2/SClassPickerDialog.h"
 #include "Misc/MessageDialog.h"
 #include "SceneStateBlueprint.h"
 #include "SceneStateBlueprintEditor.h"
@@ -11,9 +13,40 @@
 #include "SceneStateMachineGraph.h"
 #include "SceneStateMachineGraphSchema.h"
 #include "SceneStateObject.h"
+#include "SceneStateSchema.h"
 
 #define LOCTEXT_NAMESPACE "SceneStateBlueprintFactory"
 
+namespace UE::SceneState::Editor
+{
+
+class FClassViewFilter : public IClassViewerFilter
+{
+public:
+	FClassViewFilter()
+		: AllowedChildrenOfClasses({ USceneStateSchema::StaticClass() })
+	{
+	}
+
+	virtual bool IsClassAllowed(const FClassViewerInitializationOptions& InInitOptions, const UClass* InClass, TSharedRef<FClassViewerFilterFuncs> InFilterFuncs) override
+	{
+		return !InClass->HasAnyClassFlags(EClassFlags::CLASS_Deprecated | EClassFlags::CLASS_Abstract)
+			&& InFilterFuncs->IfInChildOfClassesSet(AllowedChildrenOfClasses, InClass);
+	}
+
+	virtual bool IsUnloadedClassAllowed(const FClassViewerInitializationOptions& InInitOptions, const TSharedRef<const IUnloadedBlueprintData> InUnloadedClassData, TSharedRef<FClassViewerFilterFuncs> InFilterFuncs) override
+	{
+		return !InUnloadedClassData->HasAnyClassFlags(EClassFlags::CLASS_Deprecated | EClassFlags::CLASS_Abstract)
+			&& InFilterFuncs->IfInChildOfClassesSet(AllowedChildrenOfClasses, InUnloadedClassData);
+	}
+
+private:
+	/** All children of these classes will be included unless filtered out by another setting. */
+	TSet<const UClass*> AllowedChildrenOfClasses;
+};
+
+} // UE::SceneState::Editor
+
 USceneStateBlueprintFactory::USceneStateBlueprintFactory()
 {
 	ParentClass = USceneStateObject::StaticClass();
@@ -63,6 +96,29 @@ uint32 USceneStateBlueprintFactory::GetMenuCategories() const
 	return AssetTools.FindAdvancedAssetCategory("MotionDesignCategory");
 }
 
+bool USceneStateBlueprintFactory::ConfigureProperties()
+{
+	SchemaClass = nullptr;
+
+	FClassViewerModule& ClassViewerModule = FModuleManager::LoadModuleChecked<FClassViewerModule>("ClassViewer");
+
+	FClassViewerInitializationOptions Options;
+	Options.Mode = EClassViewerMode::ClassPicker;
+	Options.DisplayMode = EClassViewerDisplayMode::TreeView;
+	Options.NameTypeToDisplay = EClassViewerNameTypeToDisplay::DisplayName;
+	Options.ClassFilters.Emplace(MakeShared<UE::SceneState::Editor::FClassViewFilter>());
+
+	const FText TitleText = LOCTEXT("SelectSceneStateSchema", "Select Scene State Schema");
+
+	UClass* ChosenSchema = nullptr;
+	if (SClassPickerDialog::PickClass(TitleText, Options, ChosenSchema, USceneStateSchema::StaticClass()))
+	{
+		SchemaClass = ChosenSchema;
+		return true;
+	}
+	return false;
+}
+
 UObject* USceneStateBlueprintFactory::FactoryCreateNew(UClass* InClass, UObject* InParent, FName InName, EObjectFlags InFlags, UObject* InContext, FFeedbackContext* InWarn)
 {
 	return FactoryCreateNew(InClass, InParent, InName, InFlags, InContext, InWarn, NAME_None);
@@ -96,6 +152,8 @@ UObject* USceneStateBlueprintFactory::FactoryCreateNew(UClass* InClass, UObject*
 			"Ensure that this Scene State Blueprint class has a Scene State compiler registered via ISceneStateBlueprintEditorModule")
 		, *GetNameSafe(Blueprint->GetClass()));
 
+	Blueprint->SetSceneStateSchema(SchemaClass);
+
 	UEdGraph* StateMachineGraph = USceneStateBlueprintFactory::AddStateMachine(Blueprint);
 
 	Blueprint->LastEditedDocuments.AddUnique(StateMachineGraph);
diff --git a/Engine/Plugins/VirtualProduction/SceneState/Source/SceneStateBlueprintEditor/Public/SceneStateBlueprintFactory.h b/Engine/Plugins/VirtualProduction/SceneState/Source/SceneStateBlueprintEditor/Public/SceneStateBlueprintFactory.h
--- a/Engine/Plugins/VirtualProduction/SceneState/Source/SceneStateBlueprintEditor/Public/SceneStateBlueprintFactory.h
+++ b/Engine/Plugins/VirtualProduction/SceneState/Source/SceneStateBlueprintEditor/Public/SceneStateBlueprintFactory.h
@@ -9,6 +9,7 @@
 class UEdGraph;
 class USceneStateBlueprint;
 class USceneStateObject;
+class USceneStateSchema;
 
 UCLASS(MinimalAPI)
 class USceneStateBlueprintFactory : public UFactory
@@ -24,10 +25,15 @@ public:
 	virtual FText GetDisplayName() const override;
 	virtual FString GetDefaultNewAssetName() const override;
 	virtual uint32 GetMenuCategories() const override;
+	virtual bool ConfigureProperties() override;
 	virtual UObject* FactoryCreateNew(UClass* InClass, UObject* InParent, FName InName, EObjectFlags InFlags, UObject* InContext, FFeedbackContext* InWarn) override;
 	virtual UObject* FactoryCreateNew(UClass* InClass, UObject* InParent, FName InName, EObjectFlags InFlags, UObject* InContext, FFeedbackContext* InWarn, FName InCallingContext) override;
 	//~ End UFactory
 
+	/** Scene state schema to use*/
+	UPROPERTY()
+	TSubclassOf<USceneStateSchema> SchemaClass;
+
 	UPROPERTY()
 	TSubclassOf<USceneStateObject> ParentClass;
 };
diff --git a/Engine/Plugins/VirtualProduction/SceneState/Source/SceneStateGameplay/Private/SceneStateGameplaySchema.cpp b/Engine/Plugins/VirtualProduction/SceneState/Source/SceneStateGameplay/Private/SceneStateGameplaySchema.cpp
new file mode 100644
--- /dev/null
+++ b/Engine/Plugins/VirtualProduction/SceneState/Source/SceneStateGameplay/Private/SceneStateGameplaySchema.cpp
@@ -0,0 +1,9 @@
+// Copyright Epic Games, Inc. All Rights Reserved.
+
+#include "SceneStateGameplaySchema.h"
+#include "GameFramework/Actor.h"
+
+TSubclassOf<UObject> USceneStateGameplaySchema::OnGetContextObjectClass() const
+{
+	return AActor::StaticClass();
+}
diff --git a/Engine/Plugins/VirtualProduction/SceneState/Source/SceneStateGameplay/Public/SceneStateGameplaySchema.h b/Engine/Plugins/VirtualProduction/SceneState/Source/SceneStateGameplay/Public/SceneStateGameplaySchema.h
new file mode 100644
--- /dev/null
+++ b/Engine/Plugins/VirtualProduction/SceneState/Source/SceneStateGameplay/Public/SceneStateGameplaySchema.h
@@ -0,0 +1,23 @@
+// Copyright Epic Games, Inc. All Rights Reserved.
+
+#pragma once
+
+#include "SceneStateSchema.h"
+#include "Templates/SubclassOf.h"
+#include "SceneStateGameplaySchema.generated.h"
+
+#define UE_API SCENESTATEGAMEPLAY_API
+
+/** Schema for scene states played via the Scene State Component / Actor context */
+UCLASS(MinimalAPI)
+class USceneStateGameplaySchema : public USceneStateSchema
+{
+	GENERATED_BODY()
+
+protected:
+	//~ Begin USceneStateSchema
+	UE_API virtual TSubclassOf<UObject> OnGetContextObjectClass() const override;
+	//~ End USceneStateSchema
+};
+
+#undef UE_API
diff --git a/Engine/Plugins/VirtualProduction/SceneState/Source/SceneStateGameplayEditor/Private/SceneStateActorFactory.cpp b/Engine/Plugins/VirtualProduction/SceneState/Source/SceneStateGameplayEditor/Private/SceneStateActorFactory.cpp
--- a/Engine/Plugins/VirtualProduction/SceneState/Source/SceneStateGameplayEditor/Private/SceneStateActorFactory.cpp
+++ b/Engine/Plugins/VirtualProduction/SceneState/Source/SceneStateGameplayEditor/Private/SceneStateActorFactory.cpp
@@ -3,6 +3,7 @@
 #include "SceneStateActorFactory.h"
 #include "SceneStateActor.h"
 #include "SceneStateBlueprint.h"
+#include "SceneStateGameplaySchema.h"
 #include "SceneStateGeneratedClass.h"
 #include "SceneStateObject.h"
 
@@ -16,14 +17,31 @@ USceneStateActorFactory::USceneStateActorFactory()
 
 bool USceneStateActorFactory::CanCreateActorFrom(const FAssetData& InAssetData, FText& OutErrorMessage)
 {
-	if (UActorFactory::CanCreateActorFrom(InAssetData, OutErrorMessage))
+	if (!InAssetData.IsValid() || !InAssetData.IsInstanceOf<USceneStateBlueprint>())
 	{
+		OutErrorMessage = LOCTEXT("InvalidSceneStateAsset", "A valid Scene State asset must be specified.");
+		return false;
+	}
+
+	const FAssetTagValueRef SchemaClassTagValue = InAssetData.TagsAndValues.FindTag(TEXT("SceneStateSchemaClass"));
+	if (!SchemaClassTagValue.IsSet())
+	{
+		// Allow creating from older scene states with no schema tag.
 		return true;
 	}
 
-	if (InAssetData.IsValid() && !InAssetData.IsInstanceOf<USceneStateBlueprint>())
+	const UClass* const SchemaClass = FindObject<UClass>(nullptr, SchemaClassTagValue.AsString());
+	if (!SchemaClass)
 	{
-		OutErrorMessage = LOCTEXT("InvalidSceneStateAsset", "A valid Scene State asset must be specified.");
+		OutErrorMessage = FText::Format(LOCTEXT("NullSchemaClass", "Schema path '{0}' is not found.")
+			, FText::FromString(SchemaClassTagValue.AsString()));
+		return false;
+	}
+
+	if (!SchemaClass->IsChildOf<USceneStateGameplaySchema>())
+	{
+		OutErrorMessage = FText::Format(LOCTEXT("InvalidSchemaClass", "Schema class '{0}' is not child of Scene State Gameplay Schema.")
+			, FText::FromString(SchemaClassTagValue.AsString()));
 		return false;
 	}
 
diff --git a/Engine/Plugins/VirtualProduction/SceneState/Source/SceneStateMachineGraph/Private/Actions/SceneStateMachineAction_NewBlueprintTask.cpp b/Engine/Plugins/VirtualProduction/SceneState/Source/SceneStateMachineGraph/Private/Actions/SceneStateMachineAction_NewBlueprintTask.cpp
--- a/Engine/Plugins/VirtualProduction/SceneState/Source/SceneStateMachineGraph/Private/Actions/SceneStateMachineAction_NewBlueprintTask.cpp
+++ b/Engine/Plugins/VirtualProduction/SceneState/Source/SceneStateMachineGraph/Private/Actions/SceneStateMachineAction_NewBlueprintTask.cpp
@@ -9,10 +9,10 @@
 namespace UE::SceneState::Graph
 {
 
-FStateMachineAction_NewBlueprintTask::FStateMachineAction_NewBlueprintTask(const FAssetData& InTaskAsset, int32 InGrouping)
+FStateMachineAction_NewBlueprintTask::FStateMachineAction_NewBlueprintTask(const FAssetData& InTaskAsset, int32 InSectionId)
 	: BlueprintTaskAsset(InTaskAsset)
 {
-	Grouping = InGrouping;
+	SectionID = InSectionId;
 
 	auto GetTag = [&InTaskAsset](FName InTagKey, const FText& InDefaultValue)
 		{
diff --git a/Engine/Plugins/VirtualProduction/SceneState/Source/SceneStateMachineGraph/Private/Actions/SceneStateMachineAction_NewBlueprintTask.h b/Engine/Plugins/VirtualProduction/SceneState/Source/SceneStateMachineGraph/Private/Actions/SceneStateMachineAction_NewBlueprintTask.h
--- a/Engine/Plugins/VirtualProduction/SceneState/Source/SceneStateMachineGraph/Private/Actions/SceneStateMachineAction_NewBlueprintTask.h
+++ b/Engine/Plugins/VirtualProduction/SceneState/Source/SceneStateMachineGraph/Private/Actions/SceneStateMachineAction_NewBlueprintTask.h
@@ -15,7 +15,7 @@ struct FStateMachineAction_NewBlueprintTask : public FEdGraphSchemaAction
 {
 	FStateMachineAction_NewBlueprintTask() = default;
 
-	explicit FStateMachineAction_NewBlueprintTask(const FAssetData& InBlueprintTaskAsset, int32 InGrouping);
+	explicit FStateMachineAction_NewBlueprintTask(const FAssetData& InBlueprintTaskAsset, int32 InSectionId);
 
 	//~ Begin FEdGraphSchemaAction
 	virtual UEdGraphNode* PerformAction(UEdGraph* InParentGraph, UEdGraphPin* InSourcePin, const FVector2f& InLocation, bool bInSelectNewNode = true) override;
diff --git a/Engine/Plugins/VirtualProduction/SceneState/Source/SceneStateMachineGraph/Private/Actions/SceneStateMachineAction_NewTask.cpp b/Engine/Plugins/VirtualProduction/SceneState/Source/SceneStateMachineGraph/Private/Actions/SceneStateMachineAction_NewTask.cpp
--- a/Engine/Plugins/VirtualProduction/SceneState/Source/SceneStateMachineGraph/Private/Actions/SceneStateMachineAction_NewTask.cpp
+++ b/Engine/Plugins/VirtualProduction/SceneState/Source/SceneStateMachineGraph/Private/Actions/SceneStateMachineAction_NewTask.cpp
@@ -7,9 +7,11 @@
 namespace UE::SceneState::Graph
 {
 
-FStateMachineAction_NewTask::FStateMachineAction_NewTask(const UScriptStruct* InTaskStruct, int32 InGrouping)
+FStateMachineAction_NewTask::FStateMachineAction_NewTask(const UScriptStruct* InTaskStruct, int32 InSectionId)
 	: TaskStruct(InTaskStruct)
 {
+	SectionID = InSectionId;
+
 	UpdateSearchData(InTaskStruct->GetDisplayNameText()
 		, InTaskStruct->GetToolTipText()
 		, FText::FromString(InTaskStruct->GetMetaData(TEXT("Category")))
diff --git a/Engine/Plugins/VirtualProduction/SceneState/Source/SceneStateMachineGraph/Private/Actions/SceneStateMachineAction_NewTask.h b/Engine/Plugins/VirtualProduction/SceneState/Source/SceneStateMachineGraph/Private/Actions/SceneStateMachineAction_NewTask.h
--- a/Engine/Plugins/VirtualProduction/SceneState/Source/SceneStateMachineGraph/Private/Actions/SceneStateMachineAction_NewTask.h
+++ b/Engine/Plugins/VirtualProduction/SceneState/Source/SceneStateMachineGraph/Private/Actions/SceneStateMachineAction_NewTask.h
@@ -12,7 +12,7 @@ struct FStateMachineAction_NewTask : public FEdGraphSchemaAction
 {
 	FStateMachineAction_NewTask() = default;
 
-	explicit FStateMachineAction_NewTask(const UScriptStruct* InTaskStruct, int32 InGrouping);
+	explicit FStateMachineAction_NewTask(const UScriptStruct* InTaskStruct, int32 InSectionId);
 
 	//~ Begin FEdGraphSchemaAction
 	virtual UEdGraphNode* PerformAction(UEdGraph* InParentGraph, UEdGraphPin* InSourcePin, const FVector2f& InLocation, bool bInSelectNewNode = true) override;
diff --git a/Engine/Plugins/VirtualProduction/SceneState/Source/SceneStateMachineGraph/Private/SceneStateMachineGraphSchema.cpp b/Engine/Plugins/VirtualProduction/SceneState/Source/SceneStateMachineGraph/Private/SceneStateMachineGraphSchema.cpp
--- a/Engine/Plugins/VirtualProduction/SceneState/Source/SceneStateMachineGraph/Private/SceneStateMachineGraphSchema.cpp
+++ b/Engine/Plugins/VirtualProduction/SceneState/Source/SceneStateMachineGraph/Private/SceneStateMachineGraphSchema.cpp
@@ -9,6 +9,7 @@
 #include "Engine/Blueprint.h"
 #include "Framework/Commands/GenericCommands.h"
 #include "GraphEditorActions.h"
+#include "ISceneStateSchemaProvider.h"
 #include "Kismet2/BlueprintEditorUtils.h"
 #include "Kismet2/KismetEditorUtilities.h"
 #include "Nodes/SceneStateMachineConduitNode.h"
@@ -20,8 +21,10 @@
 #include "Nodes/SceneStateMachineTransitionNode.h"
 #include "SceneStateMachineGraph.h"
 #include "SceneStateMachineNodeConnectionType.h"
+#include "SceneStateSchema.h"
 #include "ScopedTransaction.h"
 #include "Tasks/SceneStateTask.h"
+#include "Tasks/SceneStateTaskMetadata.h"
 #include "ToolMenu.h"
 #include "ToolMenuSection.h"
 #include "UObject/UObjectIterator.h"
@@ -152,91 +155,8 @@
 
 void USceneStateMachineGraphSchema::GetGraphContextActions(FGraphContextMenuBuilder& InContextMenuBuilder) const
 {
-	using namespace UE::SceneState::Graph;
-
-	constexpr int32 Grouping = 0; 
-
-	FString DefaultCategory;
-
-	// Add State
-	Private::AddAction<FStateMachineAction_NewNode>(InContextMenuBuilder, DefaultCategory
-		, NewObject<USceneStateMachineStateNode>(InContextMenuBuilder.OwnerOfTemporaries)
-		, FText::GetEmpty()
-		, LOCTEXT("AddState", "Add State")
-		, LOCTEXT("AddStateTooltip", "A new state")
-		, Grouping);
-
-	// Add Conduit
-	Private::AddAction<FStateMachineAction_NewNode>(InContextMenuBuilder, DefaultCategory
-		, NewObject<USceneStateMachineConduitNode>(InContextMenuBuilder.OwnerOfTemporaries)
-		, FText::GetEmpty()
-		, LOCTEXT("AddConduit", "Add Conduit")
-		, LOCTEXT("AddConduitTooltip", "Add new conduit")
-		, Grouping);
-
-	// Add Entry Point (only if it doesn't already exist)
-	bool bHasEntryNode = InContextMenuBuilder.CurrentGraph->Nodes.ContainsByPredicate(
-		[](UEdGraphNode* InNode)
-		{
-			return InNode && InNode->IsA<USceneStateMachineEntryNode>();
-		});
-
-	if (!bHasEntryNode)
-	{
-		Private::AddAction<FStateMachineAction_NewNode>(InContextMenuBuilder, DefaultCategory
-			, NewObject<USceneStateMachineEntryNode>(InContextMenuBuilder.OwnerOfTemporaries)
-			, FText::GetEmpty()
-			, LOCTEXT("AddEntry", "Add Entry Point")
-			, LOCTEXT("AddEntryTooltip", "Define the state machine's entry point")
-			, Grouping);
-	}
-
-	// Add Exit Point
-	Private::AddAction<FStateMachineAction_NewNode>(InContextMenuBuilder, DefaultCategory
-		, NewObject<USceneStateMachineExitNode>(InContextMenuBuilder.OwnerOfTemporaries)
-		, FText::GetEmpty()
-		, LOCTEXT("AddExit", "Add Exit Point")
-		, LOCTEXT("AddExitTooltip", "Define a state machine's exit point")
-		, Grouping);
-
-	// Add Comment
-	if (!InContextMenuBuilder.FromPin)
-	{
-		Private::AddAction<FStateMachineAction_NewComment>(InContextMenuBuilder, DefaultCategory
-			, FText::GetEmpty()
-			, LOCTEXT("AddComment", "Add Comment")
-			, LOCTEXT("CreateCommentSelectionTooltip", "Create a resizeable comment box around selected nodes.")
-			, Grouping);
-	}
-
-	const FString TaskCategory(TEXT("Tasks"));
-	const FTopLevelAssetPath TaskBlueprintPath(TEXT("/Script/SceneStateBlueprint.SceneStateTaskBlueprint"));
-
-	TArray<FAssetData> TaskAssets;
-	IAssetRegistry& AssetRegistry = FModuleManager::LoadModuleChecked<FAssetRegistryModule>(FName("AssetRegistry")).Get();
-	AssetRegistry.GetAssetsByClass(TaskBlueprintPath, TaskAssets, true);
-
-	// Add all Blueprint Tasks
-	for (const FAssetData& TaskAsset : TaskAssets)
-	{
-		Private::AddAction<FStateMachineAction_NewBlueprintTask>(InContextMenuBuilder, TaskCategory, TaskAsset, Grouping);
-	}
-
-	const FName MD_Hidden(TEXT("Hidden"));
-
-	// Add all Native Tasks
-	for (UScriptStruct* Struct : TObjectRange<UScriptStruct>())
-	{
-		if (Struct->HasMetaData(MD_Hidden) || !Struct->IsChildOf<FSceneStateTask>())
-		{
-			continue;
-		}
-
-		Private::AddAction<FStateMachineAction_NewTask>(InContextMenuBuilder
-			, TaskCategory
-			, Struct
-			, Grouping);
-	}
+	GetGraphContextCoreActions(InContextMenuBuilder);
+	GetGraphContextTaskActions(InContextMenuBuilder);
 }
 
 void USceneStateMachineGraphSchema::GetContextMenuActions(UToolMenu* InMenu, UGraphNodeContextMenuContext* InContext) const
@@ -538,4 +458,159 @@
 	}
 }
 
+void USceneStateMachineGraphSchema::GetGraphContextCoreActions(FGraphContextMenuBuilder& InContextMenuBuilder) const
+{
+	using namespace UE::SceneState;
+
+	const FString CoreCategory; // Category left empty
+	constexpr int32 CoreSectionId = 1; 
+
+	// Add State
+	Graph::Private::AddAction<Graph::FStateMachineAction_NewNode>(InContextMenuBuilder
+		, CoreCategory
+		, NewObject<USceneStateMachineStateNode>(InContextMenuBuilder.OwnerOfTemporaries)
+		, FText::GetEmpty()
+		, LOCTEXT("AddState", "Add State")
+		, LOCTEXT("AddStateTooltip", "A new state")
+		, CoreSectionId);
+
+	// Add Conduit
+	Graph::Private::AddAction<Graph::FStateMachineAction_NewNode>(InContextMenuBuilder
+		, CoreCategory
+		, NewObject<USceneStateMachineConduitNode>(InContextMenuBuilder.OwnerOfTemporaries)
+		, FText::GetEmpty()
+		, LOCTEXT("AddConduit", "Add Conduit")
+		, LOCTEXT("AddConduitTooltip", "Add new conduit")
+		, CoreSectionId);
+
+	// Add Entry Point (only if it doesn't already exist)
+	bool bHasEntryNode = InContextMenuBuilder.CurrentGraph->Nodes.ContainsByPredicate(
+		[](UEdGraphNode* InNode)
+		{
+			return InNode && InNode->IsA<USceneStateMachineEntryNode>();
+		});
+
+	if (!bHasEntryNode)
+	{
+		Graph::Private::AddAction<Graph::FStateMachineAction_NewNode>(InContextMenuBuilder
+			, CoreCategory
+			, NewObject<USceneStateMachineEntryNode>(InContextMenuBuilder.OwnerOfTemporaries)
+			, FText::GetEmpty()
+			, LOCTEXT("AddEntry", "Add Entry Point")
+			, LOCTEXT("AddEntryTooltip", "Define the state machine's entry point")
+			, CoreSectionId);
+	}
+
+	// Add Exit Point
+	Graph::Private::AddAction<Graph::FStateMachineAction_NewNode>(InContextMenuBuilder
+		, CoreCategory
+		, NewObject<USceneStateMachineExitNode>(InContextMenuBuilder.OwnerOfTemporaries)
+		, FText::GetEmpty()
+		, LOCTEXT("AddExit", "Add Exit Point")
+		, LOCTEXT("AddExitTooltip", "Define a state machine's exit point")
+		, CoreSectionId);
+
+	// Add Comment
+	if (!InContextMenuBuilder.FromPin)
+	{
+		Graph::Private::AddAction<Graph::FStateMachineAction_NewComment>(InContextMenuBuilder
+			, CoreCategory
+			, FText::GetEmpty()
+			, LOCTEXT("AddComment", "Add Comment")
+			, LOCTEXT("CreateCommentSelectionTooltip", "Create a resizeable comment box around selected nodes.")
+			, CoreSectionId);
+	}
+}
+
+void USceneStateMachineGraphSchema::GetGraphContextTaskActions(FGraphContextMenuBuilder& InContextMenuBuilder) const
+{
+	using namespace UE::SceneState;
+
+	const ISceneStateSchemaProvider* const SchemaProvider = Cast<ISceneStateSchemaProvider>(FBlueprintEditorUtils::FindBlueprintForGraph(InContextMenuBuilder.CurrentGraph));
+	if (!SchemaProvider)
+	{
+		return;
+	}
+
+	const USceneStateSchema* const SceneStateSchema = SchemaProvider->GetSceneStateSchemaClass().GetDefaultObject();
+	if (!SceneStateSchema)
+	{
+		return;
+	}
+
+	const FString TaskCategory; // Category left empty
+	constexpr int32 TaskSectionId = 0;
+
+	const bool bContextObjectImplementsWorld = 
+		[SceneStateSchema] 
+		{
+			const UObject* const DefaultContextObject = SceneStateSchema->GetContextObjectClass().GetDefaultObject();
+			return DefaultContextObject && DefaultContextObject->ImplementsGetWorld();
+		}();
+
+	// Buffer for task struct allocation to minimize memory allocations
+	TArray<uint8, TInlineAllocator<64>> TaskBuffer;
+
+	// C++ Tasks
+	for (UScriptStruct* Struct : TObjectRange<UScriptStruct>())
+	{
+		if (Struct->HasMetaData(Editor::MD_Hidden) || !Struct->IsChildOf<FSceneStateTask>())
+		{
+			continue;
+		}
+
+		if (!SceneStateSchema->IsTaskStructAllowed(Struct))
+		{
+			continue;
+		}
+
+		// If context does not implement world, skip structs that require world context
+		if (!bContextObjectImplementsWorld && Struct->HasMetaData(Editor::MD_RequiresContextWorld))
+		{
+			continue;
+		}
+
+		// For structs that implement WithSupportSchema, create an instance of the struct and check for support
+		if (Struct->HasMetaData(Editor::MD_WithSupportsSchema))
+		{
+			const int32 StructSize = Struct->GetStructureSize(); // same as Align(PropertiesSize, MinAlignment)
+			const int32 MinAlignment = Struct->GetMinAlignment();
+
+			TaskBuffer.SetNum(StructSize + MinAlignment, EAllowShrinking::No);
+
+			uint8* const StructBytes = Align(TaskBuffer.GetData(), MinAlignment);
+			Struct->InitializeStruct(StructBytes);
+
+			const FSceneStateTask* Task = reinterpret_cast<const FSceneStateTask*>(StructBytes);
+			const bool bTaskSupportsSchema = Task->SupportsSchema(SceneStateSchema);
+
+			Struct->DestroyStruct(StructBytes);
+
+			if (!bTaskSupportsSchema)
+			{
+				continue;
+			}
+		}
+
+		Graph::Private::AddAction<Graph::FStateMachineAction_NewTask>(InContextMenuBuilder, TaskCategory, Struct, TaskSectionId);
+	}
+
+	// Blueprint Tasks
+	{
+		const FTopLevelAssetPath TaskBlueprintPath(TEXT("/Script/SceneStateBlueprint.SceneStateTaskBlueprint"));
+
+		TArray<FAssetData> TaskAssets;
+		IAssetRegistry& AssetRegistry = FModuleManager::LoadModuleChecked<FAssetRegistryModule>(FName("AssetRegistry")).Get();
+		AssetRegistry.GetAssetsByClass(TaskBlueprintPath, TaskAssets, /*bSearchSubClasses*/true);
+
+		for (const FAssetData& TaskAsset : TaskAssets)
+		{
+			if (SceneStateSchema->IsTaskAssetAllowed(TaskAsset))
+			{
+				Graph::Private::AddAction<Graph::FStateMachineAction_NewBlueprintTask>(InContextMenuBuilder, TaskCategory, TaskAsset, TaskSectionId);
+			}
+		}
+	}
+}
+
 #undef LOCTEXT_NAMESPACE
diff --git a/Engine/Plugins/VirtualProduction/SceneState/Source/SceneStateMachineGraph/Public/SceneStateMachineGraphSchema.h b/Engine/Plugins/VirtualProduction/SceneState/Source/SceneStateMachineGraph/Public/SceneStateMachineGraphSchema.h
--- a/Engine/Plugins/VirtualProduction/SceneState/Source/SceneStateMachineGraph/Public/SceneStateMachineGraphSchema.h
+++ b/Engine/Plugins/VirtualProduction/SceneState/Source/SceneStateMachineGraph/Public/SceneStateMachineGraphSchema.h
@@ -58,6 +58,13 @@ public:
 		UEdGraphPin* TargetPin;
 	};
 	void CreateConnectionWithTransition(const FTransitionConnectionParams& InParams) const;
+
+private:
+	/** Gathers the core actions for the context menu */
+	void GetGraphContextCoreActions(FGraphContextMenuBuilder& InContextMenuBuilder) const;
+
+	/** Gathers the task actions for the context menu */
+	void GetGraphContextTaskActions(FGraphContextMenuBuilder& InContextMenuBuilder) const;
 };
 
 #undef UE_API
diff --git a/Engine/Plugins/VirtualProduction/SceneState/Source/SceneStateTasks/Public/SceneStateSpawnActorTask.h b/Engine/Plugins/VirtualProduction/SceneState/Source/SceneStateTasks/Public/SceneStateSpawnActorTask.h
--- a/Engine/Plugins/VirtualProduction/SceneState/Source/SceneStateTasks/Public/SceneStateSpawnActorTask.h
+++ b/Engine/Plugins/VirtualProduction/SceneState/Source/SceneStateTasks/Public/SceneStateSpawnActorTask.h
@@ -47,7 +47,7 @@ struct FSceneStateSpawnActorTaskInstance : public FSceneStateTaskInstance
 };
 
 /** Spawns an actor of a given class */
-USTRUCT(DisplayName="Spawn Actor", Category="Core")
+USTRUCT(DisplayName="Spawn Actor", Category="Core", meta=(RequiresContextWorld))
 struct FSceneStateSpawnActorTask : public FSceneStateTask
 {
 	GENERATED_BODY()
