Motion Design Scene State: deprecated OnPostEditChange virtual as it was const, when it should've been non-const.

#Synced-CL 50540310

diff --git a/Engine/Plugins/VirtualProduction/SceneState/Source/SceneState/Private/Tasks/SceneStateTask.cpp b/Engine/Plugins/VirtualProduction/SceneState/Source/SceneState/Private/Tasks/SceneStateTask.cpp
--- a/Engine/Plugins/VirtualProduction/SceneState/Source/SceneState/Private/Tasks/SceneStateTask.cpp
+++ b/Engine/Plugins/VirtualProduction/SceneState/Source/SceneState/Private/Tasks/SceneStateTask.cpp
@@ -50,7 +50,7 @@ void FSceneStateTask::PostEditChange(UE::SceneState::FTaskEditChange& InEditChan
 {
 	if (InTaskInstance.IsValid())
 	{
-		OnPostEditChange(InEditChange, InTaskInstance);
+		OnPostTaskEditChange(InEditChange, InTaskInstance);
 	}
 }
 #endif
diff --git a/Engine/Plugins/VirtualProduction/SceneState/Source/SceneState/Private/Tasks/SceneStateTaskEditChange.cpp b/Engine/Plugins/VirtualProduction/SceneState/Source/SceneState/Private/Tasks/SceneStateTaskEditChange.cpp
new file mode 100644
--- /dev/null
+++ b/Engine/Plugins/VirtualProduction/SceneState/Source/SceneState/Private/Tasks/SceneStateTaskEditChange.cpp
@@ -0,0 +1,25 @@
+// Copyright Epic Games, Inc. All Rights Reserved.
+
+#include "Tasks/SceneStateTaskEditChange.h"
+
+namespace UE::SceneState
+{
+
+#if WITH_EDITOR
+FTaskEditChange::FTaskEditChange()
+	: FPropertyChangedEvent(nullptr)
+{
+}
+
+bool FTaskEditChange::IsTaskChange() const
+{
+	return ChangedObject == ETaskObjectType::Task;
+}
+
+bool FTaskEditChange::IsTaskInstanceChange() const
+{
+	return ChangedObject == ETaskObjectType::TaskInstance;
+}
+#endif // WITH_EDITOR
+
+} // UE::SceneState
diff --git a/Engine/Plugins/VirtualProduction/SceneState/Source/SceneState/Public/Tasks/SceneStateTask.h b/Engine/Plugins/VirtualProduction/SceneState/Source/SceneState/Public/Tasks/SceneStateTask.h
--- a/Engine/Plugins/VirtualProduction/SceneState/Source/SceneState/Public/Tasks/SceneStateTask.h
+++ b/Engine/Plugins/VirtualProduction/SceneState/Source/SceneState/Public/Tasks/SceneStateTask.h
@@ -155,7 +155,12 @@ protected:
 	 * Called when a property within a task has changed
 	 * @see PostEditChange
 	 */
-	virtual void OnPostEditChange(UE::SceneState::FTaskEditChange& InEditChange, FStructView InTaskInstance) const
+	virtual void OnPostTaskEditChange(UE::SceneState::FTaskEditChange& InEditChange, FStructView InTaskInstance)
+	{
+	}
+
+	/** DEPRECATED 5.8, "Use OnPostTaskEditChange() that is non-const instead */
+	virtual void OnPostEditChange(UE::SceneState::FTaskEditChange& InEditChange, FStructView InTaskInstance) const final
 	{
 	}
 #endif
diff --git a/Engine/Plugins/VirtualProduction/SceneState/Source/SceneState/Public/Tasks/SceneStateTaskEditChange.h b/Engine/Plugins/VirtualProduction/SceneState/Source/SceneState/Public/Tasks/SceneStateTaskEditChange.h
--- a/Engine/Plugins/VirtualProduction/SceneState/Source/SceneState/Public/Tasks/SceneStateTaskEditChange.h
+++ b/Engine/Plugins/VirtualProduction/SceneState/Source/SceneState/Public/Tasks/SceneStateTaskEditChange.h
@@ -5,6 +5,8 @@
 #include "SceneStateEnums.h"
 #include "UObject/UnrealType.h"
 
+#define UE_API SCENESTATE_API
+
 class UObject;
 
 namespace UE::SceneState
@@ -13,10 +15,14 @@ namespace UE::SceneState
 /** Describes an Edit change for a task or task instance */
 struct FTaskEditChange : FPropertyChangedEvent
 {
-	FTaskEditChange()
-		: FPropertyChangedEvent(nullptr)
-	{
-	}
+	UE_API FTaskEditChange();
+
+	/** Returns true if the FSceneStateTask changed */
+	UE_API bool IsTaskChange() const;
+
+	/** Returns true if the FSceneStateTaskInstance changed */
+	UE_API bool IsTaskInstanceChange() const;
+
 	/** Owning object of the task */
 	UObject* Outer = nullptr;
 	/** The object that was changed: whether it was the task or the task instance */
@@ -25,3 +31,5 @@ struct FTaskEditChange : FPropertyChangedEvent
 #endif
 
 } // UE::SceneState
+
+#undef UE_API
diff --git a/Engine/Plugins/VirtualProduction/SceneState/Source/SceneStateTasks/Private/SceneStateSpawnActorTask.cpp b/Engine/Plugins/VirtualProduction/SceneState/Source/SceneStateTasks/Private/SceneStateSpawnActorTask.cpp
--- a/Engine/Plugins/VirtualProduction/SceneState/Source/SceneStateTasks/Private/SceneStateSpawnActorTask.cpp
+++ b/Engine/Plugins/VirtualProduction/SceneState/Source/SceneStateTasks/Private/SceneStateSpawnActorTask.cpp
@@ -19,12 +19,11 @@ const UScriptStruct* FSceneStateSpawnActorTask::OnGetTaskInstanceType() const
 	return FInstanceDataType::StaticStruct();
 }
 
-void FSceneStateSpawnActorTask::OnPostEditChange(UE::SceneState::FTaskEditChange& InEditChange, FStructView InTaskInstance) const
+void FSceneStateSpawnActorTask::OnPostTaskEditChange(UE::SceneState::FTaskEditChange& InEditChange, FStructView InTaskInstance)
 {
 	using namespace UE::SceneState;
 
-	if (InEditChange.ChangedObject == ETaskObjectType::Task
-		&& InEditChange.GetMemberPropertyName() == GET_MEMBER_NAME_CHECKED(FSceneStateSpawnActorTask, ActorClass))
+	if (InEditChange.IsTaskChange() && InEditChange.GetMemberPropertyName() == GET_MEMBER_NAME_CHECKED(FSceneStateSpawnActorTask, ActorClass))
 	{
 		UpdateActorTemplate(InEditChange.Outer, InTaskInstance);
 	}
diff --git a/Engine/Plugins/VirtualProduction/SceneState/Source/SceneStateTasks/Public/SceneStateSpawnActorTask.h b/Engine/Plugins/VirtualProduction/SceneState/Source/SceneStateTasks/Public/SceneStateSpawnActorTask.h
--- a/Engine/Plugins/VirtualProduction/SceneState/Source/SceneStateTasks/Public/SceneStateSpawnActorTask.h
+++ b/Engine/Plugins/VirtualProduction/SceneState/Source/SceneStateTasks/Public/SceneStateSpawnActorTask.h
@@ -58,7 +58,7 @@ protected:
 	//~ Begin FSceneStateTask
 #if WITH_EDITOR
 	UE_API virtual const UScriptStruct* OnGetTaskInstanceType() const override;
-	UE_API virtual void OnPostEditChange(UE::SceneState::FTaskEditChange& InEditChange, FStructView InTaskInstance) const override;
+	UE_API virtual void OnPostTaskEditChange(UE::SceneState::FTaskEditChange& InEditChange, FStructView InTaskInstance) override;
 #endif
 	UE_API virtual void OnStart(const FSceneStateExecutionContext& InContext, FStructView InTaskInstance) const override;
 	//~ End FSceneStateTask
