Motion Design Scene State: made some schema functions editor-only, and added metadata to tasks to easily identify different categories of tasks.

#Synced-CL 50867481

diff --git a/Engine/Plugins/VirtualProduction/AvalancheSceneState/Source/AvalancheSceneState/Private/AvaSceneStateSchema.cpp b/Engine/Plugins/VirtualProduction/AvalancheSceneState/Source/AvalancheSceneState/Private/AvaSceneStateSchema.cpp
new file mode 100644
--- /dev/null
+++ b/Engine/Plugins/VirtualProduction/AvalancheSceneState/Source/AvalancheSceneState/Private/AvaSceneStateSchema.cpp
@@ -0,0 +1,22 @@
+// Copyright Epic Games, Inc. All Rights Reserved.
+
+#include "AvaSceneStateSchema.h"
+
+#if WITH_EDITOR
+#include "SceneStateUtils.h"
+#include "Tasks/SceneStateTask.h"
+#include "Tasks/SceneStateTaskMetadata.h"
+#endif
+
+#if WITH_EDITOR
+bool UAvaSceneStateSchema::OnIsTaskStructAllowed(TSubScriptStructOf<FSceneStateTask> InTaskStruct) const
+{
+	// Allow Motion Design and Core tasks
+	constexpr bool bIncludeSuperStructs = true;
+	return UE::SceneState::Editor::StructHasAnyMetaData(InTaskStruct, 
+		{
+			UE::SceneState::Editor::MD_CoreTask,
+			UE::AvaSceneState::Editor::MD_MotionDesignTask,
+		}, bIncludeSuperStructs);
+}
+#endif
diff --git a/Engine/Plugins/VirtualProduction/AvalancheSceneState/Source/AvalancheSceneState/Private/RemoteControl/AvaSceneStateRCTask.cpp b/Engine/Plugins/VirtualProduction/AvalancheSceneState/Source/AvalancheSceneState/Private/RemoteControl/AvaSceneStateRCTask.cpp
--- a/Engine/Plugins/VirtualProduction/AvalancheSceneState/Source/AvalancheSceneState/Private/RemoteControl/AvaSceneStateRCTask.cpp
+++ b/Engine/Plugins/VirtualProduction/AvalancheSceneState/Source/AvalancheSceneState/Private/RemoteControl/AvaSceneStateRCTask.cpp
@@ -90,11 +90,6 @@ const UScriptStruct* FAvaSceneStateRCTask::OnGetTaskInstanceType() const
 	return FInstanceDataType::StaticStruct();
 }
 
-bool FAvaSceneStateRCTask::OnSupportsSchema(TNotNull<const USceneStateSchema*> InSchema) const
-{
-	return Super::OnSupportsSchema(InSchema) && InSchema->IsA<UAvaSceneStateSchema>();
-}
-
 void FAvaSceneStateRCTask::OnBuildTaskInstance(UObject* InOuter, FStructView InTaskInstance) const
 {
 	FInstanceDataType& Instance = InTaskInstance.Get<FInstanceDataType>();
diff --git a/Engine/Plugins/VirtualProduction/AvalancheSceneState/Source/AvalancheSceneState/Private/Sequence/AvaSceneStatePlaySequenceTask.cpp b/Engine/Plugins/VirtualProduction/AvalancheSceneState/Source/AvalancheSceneState/Private/Sequence/AvaSceneStatePlaySequenceTask.cpp
--- a/Engine/Plugins/VirtualProduction/AvalancheSceneState/Source/AvalancheSceneState/Private/Sequence/AvaSceneStatePlaySequenceTask.cpp
+++ b/Engine/Plugins/VirtualProduction/AvalancheSceneState/Source/AvalancheSceneState/Private/Sequence/AvaSceneStatePlaySequenceTask.cpp
@@ -16,11 +16,6 @@ const UScriptStruct* FAvaSceneStatePlaySequenceTask::OnGetTaskInstanceType() con
 {
 	return FInstanceDataType::StaticStruct();
 }
-
-bool FAvaSceneStatePlaySequenceTask::OnSupportsSchema(TNotNull<const USceneStateSchema*> InSchema) const
-{
-	return Super::OnSupportsSchema(InSchema) && InSchema->IsA<UAvaSceneStateSchema>();
-}
 #endif
 
 void FAvaSceneStatePlaySequenceTask::OnStart(const FSceneStateExecutionContext& InContext, FStructView InTaskInstance) const
diff --git a/Engine/Plugins/VirtualProduction/AvalancheSceneState/Source/AvalancheSceneState/Private/Ticker/AvaSceneStateSpawnerTickerTask.cpp b/Engine/Plugins/VirtualProduction/AvalancheSceneState/Source/AvalancheSceneState/Private/Ticker/AvaSceneStateSpawnerTickerTask.cpp
--- a/Engine/Plugins/VirtualProduction/AvalancheSceneState/Source/AvalancheSceneState/Private/Ticker/AvaSceneStateSpawnerTickerTask.cpp
+++ b/Engine/Plugins/VirtualProduction/AvalancheSceneState/Source/AvalancheSceneState/Private/Ticker/AvaSceneStateSpawnerTickerTask.cpp
@@ -12,11 +12,6 @@ const UScriptStruct* FAvaSceneStateSpawnerTickerTask::OnGetTaskInstanceType() co
 {
 	return FInstanceDataType::StaticStruct();
 }
-
-bool FAvaSceneStateSpawnerTickerTask::OnSupportsSchema(TNotNull<const USceneStateSchema*> InSchema) const
-{
-	return Super::OnSupportsSchema(InSchema) && InSchema->IsA<UAvaSceneStateSchema>();
-}
 #endif
 
 bool FAvaSceneStateSpawnerTickerTask::ShouldSpawnActor(FStructView InTaskInstance, FText& OutErrorMessage) const
diff --git a/Engine/Plugins/VirtualProduction/AvalancheSceneState/Source/AvalancheSceneState/Public/AvaSceneStateSchema.h b/Engine/Plugins/VirtualProduction/AvalancheSceneState/Source/AvalancheSceneState/Public/AvaSceneStateSchema.h
--- a/Engine/Plugins/VirtualProduction/AvalancheSceneState/Source/AvalancheSceneState/Public/AvaSceneStateSchema.h
+++ b/Engine/Plugins/VirtualProduction/AvalancheSceneState/Source/AvalancheSceneState/Public/AvaSceneStateSchema.h
@@ -5,9 +5,28 @@
 #include "SceneStateGameplaySchema.h"
 #include "AvaSceneStateSchema.generated.h"
 
+#define UE_API AVALANCHESCENESTATE_API
+
+#if WITH_EDITOR
+namespace UE::AvaSceneState::Editor
+{
+	/** Metadata to identify a task as a Motion Design task */
+	const FLazyName MD_MotionDesignTask(TEXT("MotionDesignTask"));
+}
+#endif
+
 /** Schema for Motion Design scene states. */
 UCLASS(MinimalAPI, DisplayName="Motion Design Scene State Schema")
 class UAvaSceneStateSchema : public USceneStateGameplaySchema
 {
 	GENERATED_BODY()
+
+protected:
+	//~ Begin USceneStateSchema
+#if WITH_EDITOR
+	UE_API virtual bool OnIsTaskStructAllowed(TSubScriptStructOf<FSceneStateTask> InTaskStruct) const override;
+#endif
+	//~ End USceneStateSchema
 };
+
+#undef UE_API
diff --git a/Engine/Plugins/VirtualProduction/AvalancheSceneState/Source/AvalancheSceneState/Public/RemoteControl/AvaSceneStateRCTask.h b/Engine/Plugins/VirtualProduction/AvalancheSceneState/Source/AvalancheSceneState/Public/RemoteControl/AvaSceneStateRCTask.h
--- a/Engine/Plugins/VirtualProduction/AvalancheSceneState/Source/AvalancheSceneState/Public/RemoteControl/AvaSceneStateRCTask.h
+++ b/Engine/Plugins/VirtualProduction/AvalancheSceneState/Source/AvalancheSceneState/Public/RemoteControl/AvaSceneStateRCTask.h
@@ -54,7 +54,7 @@ struct FAvaSceneStateRCTaskInstance : public FSceneStateTaskInstance
 	FInstancedPropertyBag ControllerValues;
 };
 
-USTRUCT(DisplayName="Set RC Controller Values", Category="Motion Design", meta=(RequiresContextWorld, WithSupportsSchema))
+USTRUCT(DisplayName="Set RC Controller Values", Category="Motion Design", meta=(MotionDesignTask, RequiresContextWorld))
 struct FAvaSceneStateRCTask : public FSceneStateTask
 {
 	GENERATED_BODY()
@@ -67,7 +67,6 @@ protected:
 	//~ Begin FSceneStateTask
 #if WITH_EDITOR
 	virtual const UScriptStruct* OnGetTaskInstanceType() const override;
-	virtual bool OnSupportsSchema(TNotNull<const USceneStateSchema*> InSchema) const override;
 	virtual void OnBuildTaskInstance(UObject* InOuter, FStructView InTaskInstance) const override;
 #endif
 	virtual const FSceneStateTaskBindingExtension* OnGetBindingExtension() const override;
diff --git a/Engine/Plugins/VirtualProduction/AvalancheSceneState/Source/AvalancheSceneState/Public/Sequence/AvaSceneStatePlaySequenceTask.h b/Engine/Plugins/VirtualProduction/AvalancheSceneState/Source/AvalancheSceneState/Public/Sequence/AvaSceneStatePlaySequenceTask.h
--- a/Engine/Plugins/VirtualProduction/AvalancheSceneState/Source/AvalancheSceneState/Public/Sequence/AvaSceneStatePlaySequenceTask.h
+++ b/Engine/Plugins/VirtualProduction/AvalancheSceneState/Source/AvalancheSceneState/Public/Sequence/AvaSceneStatePlaySequenceTask.h
@@ -49,7 +49,7 @@ struct FAvaSceneStatePlaySequenceTaskInstance : public FSceneStateTaskInstance
 	FDelegateHandle OnSequenceFinishedHandle;
 };
 
-USTRUCT(DisplayName="Play Sequence", Category="Motion Design", meta=(RequiresContextWorld, WithSupportsSchema, ToolTip="Plays a Motion Design Sequence"))
+USTRUCT(DisplayName="Play Sequence", Category="Motion Design", meta=(MotionDesignTask, RequiresContextWorld, ToolTip="Plays a Motion Design Sequence"))
 struct FAvaSceneStatePlaySequenceTask : public FSceneStateTask
 {
 	GENERATED_BODY()
@@ -60,7 +60,6 @@ protected:
 	//~ Begin FSceneStateTask
 #if WITH_EDITOR
 	virtual const UScriptStruct* OnGetTaskInstanceType() const override;
-	virtual bool OnSupportsSchema(TNotNull<const USceneStateSchema*> InSchema) const override;
 #endif
 	virtual void OnStart(const FSceneStateExecutionContext& InContext, FStructView InTaskInstance) const override;
 	virtual void OnStop(const FSceneStateExecutionContext& InContext, FStructView InTaskInstance, ESceneStateTaskStopReason InStopReason) const override;
diff --git a/Engine/Plugins/VirtualProduction/AvalancheSceneState/Source/AvalancheSceneState/Public/Ticker/AvaSceneStateSpawnerTickerTask.h b/Engine/Plugins/VirtualProduction/AvalancheSceneState/Source/AvalancheSceneState/Public/Ticker/AvaSceneStateSpawnerTickerTask.h
--- a/Engine/Plugins/VirtualProduction/AvalancheSceneState/Source/AvalancheSceneState/Public/Ticker/AvaSceneStateSpawnerTickerTask.h
+++ b/Engine/Plugins/VirtualProduction/AvalancheSceneState/Source/AvalancheSceneState/Public/Ticker/AvaSceneStateSpawnerTickerTask.h
@@ -21,7 +21,7 @@ struct FAvaSceneStateTickerTaskInstance : public FSceneStateSpawnActorTaskInstan
  * Spawns an actor and adds it to the Ticker.
  * Note: if the Ticker actor is invalid or does not contain a Ticker Component, the actor will not spawn.
  */
-USTRUCT(DisplayName="Spawn Actor to Ticker", Category="Motion Design", meta=(WithSupportsSchema))
+USTRUCT(DisplayName="Spawn Actor to Ticker", Category="Motion Design", meta=(MotionDesignTask))
 struct FAvaSceneStateSpawnerTickerTask : public FSceneStateSpawnActorTask
 {
 	GENERATED_BODY()
@@ -32,7 +32,6 @@ protected:
 	//~ Begin FSceneStateTask
 #if WITH_EDITOR
 	UE_API virtual const UScriptStruct* OnGetTaskInstanceType() const override;
-	UE_API virtual bool OnSupportsSchema(TNotNull<const USceneStateSchema*> InSchema) const override;
 #endif
 	//~ End FSceneStateTask
 
diff --git a/Engine/Plugins/VirtualProduction/SceneState/Source/SceneState/Private/SceneStateSchema.cpp b/Engine/Plugins/VirtualProduction/SceneState/Source/SceneState/Private/SceneStateSchema.cpp
--- a/Engine/Plugins/VirtualProduction/SceneState/Source/SceneState/Private/SceneStateSchema.cpp
+++ b/Engine/Plugins/VirtualProduction/SceneState/Source/SceneState/Private/SceneStateSchema.cpp
@@ -3,6 +3,7 @@
 #include "SceneStateSchema.h"
 #include "Tasks/SceneStateTask.h"
 
+#if WITH_EDITOR
 bool USceneStateSchema::IsTaskStructAllowed(TSubScriptStructOf<FSceneStateTask> InTaskStruct) const
 {
 	return OnIsTaskStructAllowed(InTaskStruct);
@@ -12,6 +13,7 @@ bool USceneStateSchema::IsTaskAssetAllowed(const FAssetData& InTaskAsset) const
 {
 	return OnIsTaskAssetAllowed(InTaskAsset);
 }
+#endif
 
 TSubclassOf<UObject> USceneStateSchema::GetContextObjectClass() const
 {
diff --git a/Engine/Plugins/VirtualProduction/SceneState/Source/SceneState/Private/SceneStateUtils.cpp b/Engine/Plugins/VirtualProduction/SceneState/Source/SceneState/Private/SceneStateUtils.cpp
--- a/Engine/Plugins/VirtualProduction/SceneState/Source/SceneState/Private/SceneStateUtils.cpp
+++ b/Engine/Plugins/VirtualProduction/SceneState/Source/SceneState/Private/SceneStateUtils.cpp
@@ -217,4 +217,35 @@ bool ReplaceObject(UObject*& InOutObject
 	return true;
 }
 
+#if WITH_EDITOR
+namespace Editor
+{
+
+bool StructHasMetaData(const UStruct* InStruct, FName InMetaData, bool bInIncludeSuperStructs)
+{
+	return InStruct && StructHasAnyMetaData(InStruct, MakeArrayView(&InMetaData, 1), bInIncludeSuperStructs);
+}
+
+bool StructHasAnyMetaData(const UStruct* InStruct, TConstArrayView<FName> InMetaData, bool bInIncludeSuperStructs)
+{
+	for (const UStruct* Struct = InStruct; Struct; Struct = Struct->GetSuperStruct())
+	{
+		for (const FName& MetaData : InMetaData)
+		{
+			if (Struct->HasMetaData(MetaData))
+			{
+				return true;
+			}
+		}
+		if (!bInIncludeSuperStructs)
+		{
+			return false;
+		}
+	}
+	return false;
+}
+
+} // UE::SceneState::Editor
+#endif
+
 } // UE::SceneState
diff --git a/Engine/Plugins/VirtualProduction/SceneState/Source/SceneState/Public/SceneStateSchema.h b/Engine/Plugins/VirtualProduction/SceneState/Source/SceneState/Public/SceneStateSchema.h
--- a/Engine/Plugins/VirtualProduction/SceneState/Source/SceneState/Public/SceneStateSchema.h
+++ b/Engine/Plugins/VirtualProduction/SceneState/Source/SceneState/Public/SceneStateSchema.h
@@ -19,16 +19,19 @@ class USceneStateSchema : public UObject
 	GENERATED_BODY()
 
 public:
+#if WITH_EDITOR
 	/** Determines whether the given task struct is allowed */
 	UE_API bool IsTaskStructAllowed(TSubScriptStructOf<FSceneStateTask> InTaskStruct) const;
 
 	/** Determines whether the given task asset (task blueprint) is allowed */
 	UE_API bool IsTaskAssetAllowed(const FAssetData& InTaskAsset) const;
+#endif
 
 	/** Return the context object class supported */
 	UE_API TSubclassOf<UObject> GetContextObjectClass() const;
 
 protected:
+#if WITH_EDITOR
 	/** Determines whether a given task struct is allowed */
 	virtual bool OnIsTaskStructAllowed(TSubScriptStructOf<FSceneStateTask> InTaskStruct) const
 	{
@@ -40,6 +43,7 @@ protected:
 	{
 		return true;
 	}
+#endif
 
 	/** Return the context object class supported. Called only once to keep the context object class constant throughout the schema. */
 	virtual TSubclassOf<UObject> OnGetContextObjectClass() const
diff --git a/Engine/Plugins/VirtualProduction/SceneState/Source/SceneState/Public/SceneStateUtils.h b/Engine/Plugins/VirtualProduction/SceneState/Source/SceneState/Public/SceneStateUtils.h
--- a/Engine/Plugins/VirtualProduction/SceneState/Source/SceneState/Public/SceneStateUtils.h
+++ b/Engine/Plugins/VirtualProduction/SceneState/Source/SceneState/Public/SceneStateUtils.h
@@ -99,6 +99,18 @@ namespace UE::SceneState
 		, const TCHAR* InObjectName
 		, const TCHAR* InContextName = nullptr
 		, TFunctionRef<void(UObject*)> InOnPreDiscardOldObject = [](UObject*){});
+	
+#if WITH_EDITOR
+namespace Editor
+{
+	/** Returns true if the struct or its super structs contains the given metadata. */
+	UE_API bool StructHasMetaData(const UStruct* InStruct, FName InMetaData, bool bInIncludeSuperStructs = true);
+
+	/** Returns true if the struct or its super structs contains any of the given metadata. */
+	UE_API bool StructHasAnyMetaData(const UStruct* InStruct, TConstArrayView<FName> InMetaData, bool bInIncludeSuperStructs = true);
+
+} // UE::SceneState::Editor
+#endif
 
 } // UE::SceneState
 
diff --git a/Engine/Plugins/VirtualProduction/SceneState/Source/SceneState/Public/Tasks/SceneStateTaskMetadata.h b/Engine/Plugins/VirtualProduction/SceneState/Source/SceneState/Public/Tasks/SceneStateTaskMetadata.h
--- a/Engine/Plugins/VirtualProduction/SceneState/Source/SceneState/Public/Tasks/SceneStateTaskMetadata.h
+++ b/Engine/Plugins/VirtualProduction/SceneState/Source/SceneState/Public/Tasks/SceneStateTaskMetadata.h
@@ -10,6 +10,9 @@ namespace UE::SceneState::Editor
 	/** The task should not appear in the list of available tasks (e.g. because it's a base task struct) */
 	const FLazyName MD_Hidden(TEXT("Hidden"));
 
+	/** The task is deemed as 'Core'. All the tasks that come with the Scene State plugin will be marked with this metadata. */
+	const FLazyName MD_CoreTask(TEXT("CoreTask"));
+
 	/**
 	 * The task requires the context object to implement GetWorld().
 	 * This is just a convenience metadata to avoid having to call SupportsSchema.
diff --git a/Engine/Plugins/VirtualProduction/SceneState/Source/SceneStateGameplay/Private/SceneStateGameplaySchema.cpp b/Engine/Plugins/VirtualProduction/SceneState/Source/SceneStateGameplay/Private/SceneStateGameplaySchema.cpp
--- a/Engine/Plugins/VirtualProduction/SceneState/Source/SceneStateGameplay/Private/SceneStateGameplaySchema.cpp
+++ b/Engine/Plugins/VirtualProduction/SceneState/Source/SceneStateGameplay/Private/SceneStateGameplaySchema.cpp
@@ -3,6 +3,20 @@
 #include "SceneStateGameplaySchema.h"
 #include "GameFramework/Actor.h"
 
+#if WITH_EDITOR
+#include "SceneStateUtils.h"
+#include "Tasks/SceneStateTask.h"
+#include "Tasks/SceneStateTaskMetadata.h"
+#endif
+
+#if WITH_EDITOR
+bool USceneStateGameplaySchema::OnIsTaskStructAllowed(TSubScriptStructOf<FSceneStateTask> InTaskStruct) const
+{
+	constexpr bool bIncludeSuperStructs = true;
+	return UE::SceneState::Editor::StructHasMetaData(InTaskStruct, UE::SceneState::Editor::MD_CoreTask, bIncludeSuperStructs);
+}
+#endif
+
 TSubclassOf<UObject> USceneStateGameplaySchema::OnGetContextObjectClass() const
 {
 	return AActor::StaticClass();
diff --git a/Engine/Plugins/VirtualProduction/SceneState/Source/SceneStateGameplay/Public/SceneStateGameplaySchema.h b/Engine/Plugins/VirtualProduction/SceneState/Source/SceneStateGameplay/Public/SceneStateGameplaySchema.h
--- a/Engine/Plugins/VirtualProduction/SceneState/Source/SceneStateGameplay/Public/SceneStateGameplaySchema.h
+++ b/Engine/Plugins/VirtualProduction/SceneState/Source/SceneStateGameplay/Public/SceneStateGameplaySchema.h
@@ -16,6 +16,9 @@ class USceneStateGameplaySchema : public USceneStateSchema
 
 protected:
 	//~ Begin USceneStateSchema
+#if WITH_EDITOR
+	UE_API virtual bool OnIsTaskStructAllowed(TSubScriptStructOf<FSceneStateTask> InTaskStruct) const override;
+#endif
 	UE_API virtual TSubclassOf<UObject> OnGetContextObjectClass() const override;
 	//~ End USceneStateSchema
 };
diff --git a/Engine/Plugins/VirtualProduction/SceneState/Source/SceneStateMachineGraph/Private/SceneStateMachineGraphSchema.cpp b/Engine/Plugins/VirtualProduction/SceneState/Source/SceneStateMachineGraph/Private/SceneStateMachineGraphSchema.cpp
--- a/Engine/Plugins/VirtualProduction/SceneState/Source/SceneStateMachineGraph/Private/SceneStateMachineGraphSchema.cpp
+++ b/Engine/Plugins/VirtualProduction/SceneState/Source/SceneStateMachineGraph/Private/SceneStateMachineGraphSchema.cpp
@@ -22,6 +22,7 @@
 #include "SceneStateMachineGraph.h"
 #include "SceneStateMachineNodeConnectionType.h"
 #include "SceneStateSchema.h"
+#include "SceneStateUtils.h"
 #include "ScopedTransaction.h"
 #include "Tasks/SceneStateTask.h"
 #include "Tasks/SceneStateTaskMetadata.h"
@@ -554,7 +555,8 @@ void USceneStateMachineGraphSchema::GetGraphContextTaskActions(FGraphContextMenu
 	// C++ Tasks
 	for (UScriptStruct* Struct : TObjectRange<UScriptStruct>())
 	{
-		if (Struct->HasMetaData(Editor::MD_Hidden) || !Struct->IsChildOf<FSceneStateTask>())
+		constexpr bool bIncludeSuperStructs = true;
+		if (Editor::StructHasMetaData(Struct, Editor::MD_Hidden, !bIncludeSuperStructs) || !Struct->IsChildOf<FSceneStateTask>())
 		{
 			continue;
 		}
@@ -565,13 +567,13 @@ void USceneStateMachineGraphSchema::GetGraphContextTaskActions(FGraphContextMenu
 		}
 
 		// If context does not implement world, skip structs that require world context
-		if (!bContextObjectImplementsWorld && Struct->HasMetaData(Editor::MD_RequiresContextWorld))
+		if (!bContextObjectImplementsWorld && Editor::StructHasMetaData(Struct, Editor::MD_RequiresContextWorld, bIncludeSuperStructs))
 		{
 			continue;
 		}
 
 		// For structs that implement WithSupportSchema, create an instance of the struct and check for support
-		if (Struct->HasMetaData(Editor::MD_WithSupportsSchema))
+		if (Editor::StructHasMetaData(Struct, Editor::MD_WithSupportsSchema, bIncludeSuperStructs))
 		{
 			const int32 StructSize = Struct->GetStructureSize(); // same as Align(PropertiesSize, MinAlignment)
 			const int32 MinAlignment = Struct->GetMinAlignment();
diff --git a/Engine/Plugins/VirtualProduction/SceneState/Source/SceneStateTasks/Public/SceneStateDelayTask.h b/Engine/Plugins/VirtualProduction/SceneState/Source/SceneStateTasks/Public/SceneStateDelayTask.h
--- a/Engine/Plugins/VirtualProduction/SceneState/Source/SceneStateTasks/Public/SceneStateDelayTask.h
+++ b/Engine/Plugins/VirtualProduction/SceneState/Source/SceneStateTasks/Public/SceneStateDelayTask.h
@@ -19,7 +19,8 @@ struct FSceneStateDelayTaskInstance : public FSceneStateTaskInstance
 	float Delay = 0.5f;
 };
 
-USTRUCT(DisplayName="Delay", Category="Core", meta=(ToolTip="Waits for a set amount of seconds"))
+/** Waits for a set amount of seconds */
+USTRUCT(DisplayName="Delay", Category="Core", meta=(CoreTask))
 struct FSceneStateDelayTask : public FSceneStateTask
 {
 	GENERATED_BODY()
diff --git a/Engine/Plugins/VirtualProduction/SceneState/Source/SceneStateTasks/Public/SceneStatePrintStringTask.h b/Engine/Plugins/VirtualProduction/SceneState/Source/SceneStateTasks/Public/SceneStatePrintStringTask.h
--- a/Engine/Plugins/VirtualProduction/SceneState/Source/SceneStateTasks/Public/SceneStatePrintStringTask.h
+++ b/Engine/Plugins/VirtualProduction/SceneState/Source/SceneStateTasks/Public/SceneStatePrintStringTask.h
@@ -46,7 +46,7 @@ struct FSceneStatePrintStringTaskInstance : public FSceneStateTaskInstance
 	FSceneStatePrintSettings PrintSettings;
 };
 
-USTRUCT(DisplayName="Print String", Category="Core")
+USTRUCT(DisplayName="Print String", Category="Core", meta=(CoreTask))
 struct FSceneStatePrintStringTask : public FSceneStateTask
 {
 	GENERATED_BODY()
diff --git a/Engine/Plugins/VirtualProduction/SceneState/Source/SceneStateTasks/Public/SceneStateSpawnActorTask.h b/Engine/Plugins/VirtualProduction/SceneState/Source/SceneStateTasks/Public/SceneStateSpawnActorTask.h
--- a/Engine/Plugins/VirtualProduction/SceneState/Source/SceneStateTasks/Public/SceneStateSpawnActorTask.h
+++ b/Engine/Plugins/VirtualProduction/SceneState/Source/SceneStateTasks/Public/SceneStateSpawnActorTask.h
@@ -47,7 +47,7 @@ struct FSceneStateSpawnActorTaskInstance : public FSceneStateTaskInstance
 };
 
 /** Spawns an actor of a given class */
-USTRUCT(DisplayName="Spawn Actor", Category="Core", meta=(RequiresContextWorld))
+USTRUCT(DisplayName="Spawn Actor", Category="Core", meta=(CoreTask, RequiresContextWorld))
 struct FSceneStateSpawnActorTask : public FSceneStateTask
 {
 	GENERATED_BODY()
diff --git a/Engine/Plugins/VirtualProduction/SceneState/Source/SceneStateTasks/Public/Setters/SceneStateSetBoolTask.h b/Engine/Plugins/VirtualProduction/SceneState/Source/SceneStateTasks/Public/Setters/SceneStateSetBoolTask.h
--- a/Engine/Plugins/VirtualProduction/SceneState/Source/SceneStateTasks/Public/Setters/SceneStateSetBoolTask.h
+++ b/Engine/Plugins/VirtualProduction/SceneState/Source/SceneStateTasks/Public/Setters/SceneStateSetBoolTask.h
@@ -22,7 +22,7 @@ struct FSceneStateSetBoolTaskInstance : public FSceneStateTaskInstance
 };
 
 /** Sets a boolean value to the bound property reference */
-USTRUCT(DisplayName="Set Boolean", Category="Setter")
+USTRUCT(DisplayName="Set Boolean", Category="Setter", meta=(CoreTask))
 struct FSceneStateSetBoolTask : public FSceneStateTask
 {
 	GENERATED_BODY()
diff --git a/Engine/Plugins/VirtualProduction/SceneState/Source/SceneStateTasks/Public/Setters/SceneStateSetFloatTask.h b/Engine/Plugins/VirtualProduction/SceneState/Source/SceneStateTasks/Public/Setters/SceneStateSetFloatTask.h
--- a/Engine/Plugins/VirtualProduction/SceneState/Source/SceneStateTasks/Public/Setters/SceneStateSetFloatTask.h
+++ b/Engine/Plugins/VirtualProduction/SceneState/Source/SceneStateTasks/Public/Setters/SceneStateSetFloatTask.h
@@ -22,7 +22,7 @@ struct FSceneStateSetFloatTaskInstance : public FSceneStateTaskInstance
 };
 
 /** Sets a float value to the bound property reference */
-USTRUCT(DisplayName="Set Float", Category="Setter")
+USTRUCT(DisplayName="Set Float", Category="Setter", meta=(CoreTask))
 struct FSceneStateSetFloatTask : public FSceneStateTask
 {
 	GENERATED_BODY()
diff --git a/Engine/Plugins/VirtualProduction/SceneState/Source/SceneStateTasks/Public/Setters/SceneStateSetIntegerTask.h b/Engine/Plugins/VirtualProduction/SceneState/Source/SceneStateTasks/Public/Setters/SceneStateSetIntegerTask.h
--- a/Engine/Plugins/VirtualProduction/SceneState/Source/SceneStateTasks/Public/Setters/SceneStateSetIntegerTask.h
+++ b/Engine/Plugins/VirtualProduction/SceneState/Source/SceneStateTasks/Public/Setters/SceneStateSetIntegerTask.h
@@ -22,7 +22,7 @@ struct FSceneStateSetIntegerTaskInstance : public FSceneStateTaskInstance
 };
 
 /** Sets a integer value to the bound property reference */
-USTRUCT(DisplayName="Set Integer", Category="Setter")
+USTRUCT(DisplayName="Set Integer", Category="Setter", meta=(CoreTask))
 struct FSceneStateSetIntegerTask : public FSceneStateTask
 {
 	GENERATED_BODY()
diff --git a/Engine/Plugins/VirtualProduction/SceneState/Source/SceneStateTasks/Public/Setters/SceneStateSetStringTask.h b/Engine/Plugins/VirtualProduction/SceneState/Source/SceneStateTasks/Public/Setters/SceneStateSetStringTask.h
--- a/Engine/Plugins/VirtualProduction/SceneState/Source/SceneStateTasks/Public/Setters/SceneStateSetStringTask.h
+++ b/Engine/Plugins/VirtualProduction/SceneState/Source/SceneStateTasks/Public/Setters/SceneStateSetStringTask.h
@@ -22,7 +22,7 @@ struct FSceneStateSetStringTaskInstance : public FSceneStateTaskInstance
 };
 
 /** Sets a string value to the bound property reference */
-USTRUCT(DisplayName="Set String", Category="Setter")
+USTRUCT(DisplayName="Set String", Category="Setter", meta=(CoreTask))
 struct FSceneStateSetStringTask : public FSceneStateTask
 {
 	GENERATED_BODY()
diff --git a/Engine/Plugins/VirtualProduction/SceneState/Source/SceneStateTasks/Public/Setters/SceneStateSetTextTask.h b/Engine/Plugins/VirtualProduction/SceneState/Source/SceneStateTasks/Public/Setters/SceneStateSetTextTask.h
--- a/Engine/Plugins/VirtualProduction/SceneState/Source/SceneStateTasks/Public/Setters/SceneStateSetTextTask.h
+++ b/Engine/Plugins/VirtualProduction/SceneState/Source/SceneStateTasks/Public/Setters/SceneStateSetTextTask.h
@@ -22,7 +22,7 @@ struct FSceneStateSetTextTaskInstance : public FSceneStateTaskInstance
 };
 
 /** Sets a text value to the bound property reference */
-USTRUCT(DisplayName="Set Text", Category="Setter")
+USTRUCT(DisplayName="Set Text", Category="Setter", meta=(CoreTask))
 struct FSceneStateSetTextTask : public FSceneStateTask
 {
 	GENERATED_BODY()
