Motion Design Scene State: added state parameters for an additional (closer) scope of parameters for tasks.

#Synced-CL 50867481

diff --git a/Engine/Plugins/VirtualProduction/SceneState/Source/SceneState/Private/SceneState.cpp b/Engine/Plugins/VirtualProduction/SceneState/Source/SceneState/Private/SceneState.cpp
--- a/Engine/Plugins/VirtualProduction/SceneState/Source/SceneState/Private/SceneState.cpp
+++ b/Engine/Plugins/VirtualProduction/SceneState/Source/SceneState/Private/SceneState.cpp
@@ -34,10 +34,13 @@ void FSceneState::Enter(const FSceneStateExecutionContext& InContext) const
 	UE_LOG(LogSceneState, Verbose, TEXT("State (%s) receiving enter"), GetStateName(InContext).GetData());
 
 	InContext.SetStateStatus(*this, Instance, EExecutionStatus::Running);
+	Instance.InitializeParameters(Parameters);
+	InContext.SetupFunctionInstances(BindingsBatch);
 
-	// Apply Event Handlers before anything else starts so the Event Data becomes available to substate machines and tasks
+	// Apply Event Handlers before anything else starts so the Event Data becomes available to this state, substate machines and tasks
 	CaptureEvents(InContext);
 	AllocateTaskInstances(InContext, InContext.GetTemplateTaskInstances(*this));
+	ApplyBindings(InContext, Instance);
 
 	for (const FSceneStateTransition& Transition : InContext.GetTransitions(*this))
 	{
@@ -140,6 +143,7 @@ void FSceneState::Exit(const FSceneStateExecutionContext& InContext) const
 
 	InContext.RemoveStateInstance(*this);
 	InContext.RemoveTaskInstanceContainer(*this);
+	InContext.RemoveFunctionInstances(BindingsBatch);
 
 	ResetCapturedEvents(InContext);
 }
@@ -250,6 +254,19 @@ void FSceneState::AllocateTaskInstances(const FSceneStateExecutionContext& InCon
 	}
 }
 
+bool FSceneState::ApplyBindings(const FSceneStateExecutionContext& InContext, FSceneStateInstance& InStateInstance) const
+{
+	QUICK_SCOPE_CYCLE_COUNTER(SceneState_ApplyBindings);
+
+	const UE::SceneState::FApplyBatchParams ApplyBatchParams
+		{
+			.BindingsBatch = BindingsBatch.Get(),
+			.TargetDataView = InStateInstance.GetParametersMutableView(),
+		};
+
+	return ApplyBatch(InContext, ApplyBatchParams);
+}
+
 void FSceneState::InstanceTaskObjects(UObject* InOuter, TConstArrayView<FStructView> InTargets, TConstArrayView<FConstStructView> InSources, TFunctionRef<UObject*(FObjectDuplicationParameters&)> InDuplicationFunc) const
 {
 	check(InTargets.Num() == InSources.Num());
diff --git a/Engine/Plugins/VirtualProduction/SceneState/Source/SceneState/Private/SceneStateExecutionContext.cpp b/Engine/Plugins/VirtualProduction/SceneState/Source/SceneState/Private/SceneStateExecutionContext.cpp
--- a/Engine/Plugins/VirtualProduction/SceneState/Source/SceneState/Private/SceneStateExecutionContext.cpp
+++ b/Engine/Plugins/VirtualProduction/SceneState/Source/SceneState/Private/SceneStateExecutionContext.cpp
@@ -204,6 +204,13 @@ FPropertyBindingDataView FSceneStateExecutionContext::FindDataView(const FSceneS
 
 	case ESceneStateDataType::Function:
 		return FindFunctionInstance(DataIndex);
+
+	case ESceneStateDataType::State:
+		if (FSceneStateInstance* StateInstance = FindStateInstance(DataIndex))
+		{
+			return StateInstance->GetParametersMutableView();
+		}
+		break;
 	}
 
 	// The data handle could be exclusively on the parent context, so call FindDataView on the parent too
diff --git a/Engine/Plugins/VirtualProduction/SceneState/Source/SceneState/Private/SceneStateInstance.cpp b/Engine/Plugins/VirtualProduction/SceneState/Source/SceneState/Private/SceneStateInstance.cpp
--- a/Engine/Plugins/VirtualProduction/SceneState/Source/SceneState/Private/SceneStateInstance.cpp
+++ b/Engine/Plugins/VirtualProduction/SceneState/Source/SceneState/Private/SceneStateInstance.cpp
@@ -26,6 +26,11 @@ FSceneStateInstance::FSceneStateInstance()
 }
 PRAGMA_ENABLE_DEPRECATION_WARNINGS
 
+void FSceneStateInstance::InitializeParameters(const FInstancedPropertyBag& InTemplateParameters)
+{
+	Parameters = InTemplateParameters;
+}
+
 UE::SceneState::EExecutionStatus FSceneStateInstance::GetStatus() const
 {
 PRAGMA_DISABLE_DEPRECATION_WARNINGS
diff --git a/Engine/Plugins/VirtualProduction/SceneState/Source/SceneState/Private/SceneStateTemplateData.cpp b/Engine/Plugins/VirtualProduction/SceneState/Source/SceneState/Private/SceneStateTemplateData.cpp
--- a/Engine/Plugins/VirtualProduction/SceneState/Source/SceneState/Private/SceneStateTemplateData.cpp
+++ b/Engine/Plugins/VirtualProduction/SceneState/Source/SceneState/Private/SceneStateTemplateData.cpp
@@ -76,6 +76,9 @@ const UStruct* USceneStateTemplateData::FindDataStruct(const UClass* InOwnerClas
 
 	case ESceneStateDataType::Function:
 		return FunctionInstances[DataIndex].GetScriptStruct();
+
+	case ESceneStateDataType::State:
+		return States[DataIndex].GetParametersStruct();
 	}
 
 	return nullptr;
diff --git a/Engine/Plugins/VirtualProduction/SceneState/Source/SceneState/Public/SceneState.h b/Engine/Plugins/VirtualProduction/SceneState/Source/SceneState/Public/SceneState.h
--- a/Engine/Plugins/VirtualProduction/SceneState/Source/SceneState/Public/SceneState.h
+++ b/Engine/Plugins/VirtualProduction/SceneState/Source/SceneState/Public/SceneState.h
@@ -2,9 +2,11 @@
 
 #pragma once
 
+#include "PropertyBindingTypes.h"
 #include "SceneStateEnums.h"
 #include "SceneStateRange.h"
 #include "SceneStateReentryGuard.h"
+#include "StructUtils/PropertyBag.h"
 #include "SceneState.generated.h"
 
 #define UE_API SCENESTATE_API
@@ -16,6 +18,7 @@ struct FStructView;
 
 namespace UE::SceneState::Editor
 {
+	class FBindingCompiler;
 	class FBlueprintCompilerContext;
 	class FStateMachineCompiler;
 
@@ -52,6 +55,11 @@ struct FSceneState
 	/** Called when the state is no longer active in its parent state machine */
 	UE_API void Exit(const FSceneStateExecutionContext& InContext) const;
 
+	const UStruct* GetParametersStruct() const
+	{
+		return Parameters.GetPropertyBagStruct();
+	}
+
 	FSceneStateRange GetTaskRange() const
 	{
 		return TaskRange;
@@ -85,6 +93,9 @@ struct FSceneState
 	UE_API void AllocateTaskInstances(const FSceneStateExecutionContext& InContext, TConstArrayView<FConstStructView> InTemplateTaskInstances) const;
 
 private:
+	/** Applies Bindings to the given state instance */
+	bool ApplyBindings(const FSceneStateExecutionContext& InContext, FSceneStateInstance& InStateInstance) const;
+
 	/** Instances the instanced references in source data into the target data using the provided duplication functor */
 	UE_API void InstanceTaskObjects(UObject* InOuter
 		, TConstArrayView<FStructView> InTargets
@@ -97,6 +108,14 @@ private:
 	/** Removes all the event data that was captured by the handlers this state owns */
 	void ResetCapturedEvents(const FSceneStateExecutionContext& InContext) const;
 
+	/** Template parameters to use for the state instances */
+	UPROPERTY()
+	FInstancedPropertyBag Parameters;
+
+	/** Bindings batch where this state is the target */
+	UPROPERTY()
+	FPropertyBindingIndex16 BindingsBatch;
+
 	/** Index and Count of the tasks that belong to this state. */
 	UPROPERTY()
 	FSceneStateRange TaskRange;
@@ -119,6 +138,7 @@ private:
 	 */
 	UE::SceneState::FReentryHandle ReentryHandle;
 
+	friend UE::SceneState::Editor::FBindingCompiler;
 	friend UE::SceneState::Editor::FBlueprintCompilerContext;
 	friend UE::SceneState::Editor::FStateMachineCompiler;
 };
diff --git a/Engine/Plugins/VirtualProduction/SceneState/Source/SceneState/Public/SceneStateInstance.h b/Engine/Plugins/VirtualProduction/SceneState/Source/SceneState/Public/SceneStateInstance.h
--- a/Engine/Plugins/VirtualProduction/SceneState/Source/SceneState/Public/SceneStateInstance.h
+++ b/Engine/Plugins/VirtualProduction/SceneState/Source/SceneState/Public/SceneStateInstance.h
@@ -3,6 +3,7 @@
 #pragma once
 
 #include "SceneStateEnums.h"
+#include "StructUtils/PropertyBag.h"
 #include "SceneStateInstance.generated.h"
 
 #define UE_API SCENESTATE_API
@@ -24,6 +25,9 @@ PRAGMA_ENABLE_DEPRECATION_WARNINGS
 	UE_DEPRECATED(5.8, "ElapsedTime is deprecated, it is no longer in use.")
 	float ElapsedTime = 0.f;
 
+	/** Initializes this instance's parameters to the template's */
+	void InitializeParameters(const FInstancedPropertyBag& InTemplateParameters);
+
 	/** Gets the execution status of the state instance */
 	UE_API UE::SceneState::EExecutionStatus GetStatus() const;
 
@@ -37,7 +41,16 @@ PRAGMA_ENABLE_DEPRECATION_WARNINGS
 	/** Get an identifier for this instance */
 	UE_INTERNAL UE_API uint16 GetInstanceId() const;
 
+	FStructView GetParametersMutableView()
+	{
+		return Parameters.GetMutableValue();
+	}
+
 private:
+	/** Instanced state parameters */
+	UPROPERTY()
+	FInstancedPropertyBag Parameters;
+
 	/** The id for this instance. This is used to differentiate state instances for a same state and consequently task instances for a same task*/
 	uint16 InstanceId;
 };
diff --git a/Engine/Plugins/VirtualProduction/SceneState/Source/SceneState/Public/SceneStateTemplateData.h b/Engine/Plugins/VirtualProduction/SceneState/Source/SceneState/Public/SceneStateTemplateData.h
--- a/Engine/Plugins/VirtualProduction/SceneState/Source/SceneState/Public/SceneStateTemplateData.h
+++ b/Engine/Plugins/VirtualProduction/SceneState/Source/SceneState/Public/SceneStateTemplateData.h
@@ -169,6 +169,10 @@ public:
 	UPROPERTY()
 	TMap<FGuid, uint16> StateMachineIdToIndex;
 
+	/** Map of the state parameters id to the index in the state array */
+	UPROPERTY()
+	TMap<FGuid, uint16> StateParameterIdToIndex;
+
 #if WITH_EDITORONLY_DATA
 	/** Holds all the state editor-only metadata. This array matches in index with the states array. */
 	UPROPERTY()
diff --git a/Engine/Plugins/VirtualProduction/SceneState/Source/SceneStateBinding/Public/SceneStateBindingDataHandle.h b/Engine/Plugins/VirtualProduction/SceneState/Source/SceneStateBinding/Public/SceneStateBindingDataHandle.h
--- a/Engine/Plugins/VirtualProduction/SceneState/Source/SceneStateBinding/Public/SceneStateBindingDataHandle.h
+++ b/Engine/Plugins/VirtualProduction/SceneState/Source/SceneStateBinding/Public/SceneStateBindingDataHandle.h
@@ -16,6 +16,7 @@ enum class ESceneStateDataType : uint8
 	StateMachine,
 	Transition,
 	Function,
+	State,
 };
 
 USTRUCT()
diff --git a/Engine/Plugins/VirtualProduction/SceneState/Source/SceneStateBlueprint/Private/SceneStateBlueprint.cpp b/Engine/Plugins/VirtualProduction/SceneState/Source/SceneStateBlueprint/Private/SceneStateBlueprint.cpp
--- a/Engine/Plugins/VirtualProduction/SceneState/Source/SceneStateBlueprint/Private/SceneStateBlueprint.cpp
+++ b/Engine/Plugins/VirtualProduction/SceneState/Source/SceneStateBlueprint/Private/SceneStateBlueprint.cpp
@@ -36,6 +36,9 @@ USceneStateBlueprint::USceneStateBlueprint(const FObjectInitializer& InObjectIni
 	OnGraphParametersChangedHandle = USceneStateMachineGraph::OnParametersChanged().AddUObject(this
 		, &USceneStateBlueprint::OnGraphParametersChanged);
 
+	OnStateParametersChangedHandle = USceneStateMachineStateNode::OnParametersChanged().AddUObject(this
+		, &USceneStateBlueprint::OnStateParametersChanged);
+
 #if WITH_EDITOR
 	BindingCollection.SetBindingsOwner(this);
 #endif // WITH_EDITOR
@@ -88,6 +91,14 @@ void USceneStateBlueprint::GetBindingStructs(const FGetBindableStructsParams& In
 		return;
 	}
 
+	// Add State Bindable Structs if the Target Struct Id matches the state node parameters id
+	if (USceneStateMachineStateNode* StateNode = FindStateMatchingId(*this, InParams.TargetStructId))
+	{
+		GatherDefaultBindingDescs();
+		GetStateBindingDescs(StateNode, OutBindingDescs, FString());
+		return;
+	}
+
 	// Add Transition Bindable Structs if the Target Struct Id matches the Transition Parameters Id
 	if (USceneStateMachineTransitionNode* TransitionNode = FindTransitionMatchingId(*this, InParams.TargetStructId))
 	{
@@ -233,8 +244,8 @@ bool USceneStateBlueprint::CanCreateParameter(const FGuid InStructId) const
 {
 	using namespace UE::SceneState::Graph;
 
-	// Only support creating parameters in BP Variables and State machine parameters
-	return InStructId == GetRootId() || FindStateMachineMatchingId(*this, InStructId);
+	// Only support creating parameters in BP Variables, State machine parameters and state parameters.
+	return InStructId == GetRootId() || FindStateMachineMatchingId(*this, InStructId) || FindStateMatchingId(*this, InStructId);
 }
 
 void USceneStateBlueprint::CreateParametersForStruct(const FGuid InStructId, TArrayView<UE::PropertyBinding::FPropertyCreationDescriptor> InOutCreationDescs)
@@ -250,6 +261,12 @@ void USceneStateBlueprint::CreateParametersForStruct(const FGuid InStructId, TAr
 	if (USceneStateMachineGraph* StateMachineGraph = FindStateMachineMatchingId(*this, InStructId))
 	{
 		CreateUniquelyNamedPropertiesInPropertyBag(InOutCreationDescs, StateMachineGraph->Parameters);
+		return;
+	}
+
+	if (USceneStateMachineStateNode* StateNode = FindStateMatchingId(*this, InStructId))
+	{
+		CreateUniquelyNamedPropertiesInPropertyBag(InOutCreationDescs, StateNode->Parameters);
 	}
 }
 
@@ -348,6 +365,9 @@ void USceneStateBlueprint::BeginDestroy()
 
 	USceneStateMachineGraph::OnParametersChanged().Remove(OnGraphParametersChangedHandle);
 	OnGraphParametersChangedHandle.Reset();
+
+	USceneStateMachineStateNode::OnParametersChanged().Remove(OnStateParametersChangedHandle);
+	OnStateParametersChangedHandle.Reset();
 }
 
 void USceneStateBlueprint::OnRenameVariableReferences(UBlueprint* InBlueprint, UClass* InVariableClass, const FName& InOldVariableName, const FName& InNewVariableName)
@@ -399,18 +419,38 @@ void USceneStateBlueprint::RenameVariableReferenceInPath(FPropertyBindingPath& I
 
 void USceneStateBlueprint::OnGraphParametersChanged(USceneStateMachineGraph* InGraph)
 {
+	if (!InGraph)
+	{
+		return;
+	}
+
 	BindingCollection.ForEachMutableBinding(
 		[This=this, InGraph](FPropertyBindingBinding& InBinding)
 		{
-			This->UpdateGraphParametersBindings(InBinding.GetMutableSourcePath(), InGraph);
-			This->UpdateGraphParametersBindings(InBinding.GetMutableTargetPath(), InGraph);
+			This->UpdateParametersBindings(InBinding.GetMutableSourcePath(), InGraph->ParametersId, InGraph->Parameters);
+			This->UpdateParametersBindings(InBinding.GetMutableTargetPath(), InGraph->ParametersId, InGraph->Parameters);
+		});
+}
+
+void USceneStateBlueprint::OnStateParametersChanged(USceneStateMachineStateNode* InStateNode)
+{
+	if (!InStateNode)
+	{
+		return;
+	}
+
+	BindingCollection.ForEachMutableBinding(
+		[This=this, InStateNode](FPropertyBindingBinding& InBinding)
+		{
+			This->UpdateParametersBindings(InBinding.GetMutableSourcePath(), InStateNode->ParametersId, InStateNode->Parameters);
+			This->UpdateParametersBindings(InBinding.GetMutableTargetPath(), InStateNode->ParametersId, InStateNode->Parameters);
 		});
 }
 
-void USceneStateBlueprint::UpdateGraphParametersBindings(FPropertyBindingPath& InPath, USceneStateMachineGraph* InGraph)
+void USceneStateBlueprint::UpdateParametersBindings(FPropertyBindingPath& InPath, const FGuid& InParametersId, const FInstancedPropertyBag& InParameters)
 {
 	// Only consider fixing paths that are set to the graph parameters
-	if (InPath.GetStructID() != InGraph->ParametersId)
+	if (InPath.GetStructID() != InParametersId)
 	{
 		return;
 	}
@@ -424,7 +464,7 @@ void USceneStateBlueprint::UpdateGraphParametersBindings(FPropertyBindingPath& I
 	// Only need to consider the first segment of the path (i.e. the segment that might've been renamed)
 	FPropertyBindingPathSegment& Segment = Segments[0];
 
-	if (const FPropertyBagPropertyDesc* Property = InGraph->Parameters.FindPropertyDescByID(Segment.GetPropertyGuid()))
+	if (const FPropertyBagPropertyDesc* Property = InParameters.FindPropertyDescByID(Segment.GetPropertyGuid()))
 	{
 		const FName SegmentName = Segment.GetName();
 
diff --git a/Engine/Plugins/VirtualProduction/SceneState/Source/SceneStateBlueprint/Private/SceneStateBlueprintBindingUtils.cpp b/Engine/Plugins/VirtualProduction/SceneState/Source/SceneStateBlueprint/Private/SceneStateBlueprintBindingUtils.cpp
--- a/Engine/Plugins/VirtualProduction/SceneState/Source/SceneStateBlueprint/Private/SceneStateBlueprintBindingUtils.cpp
+++ b/Engine/Plugins/VirtualProduction/SceneState/Source/SceneStateBlueprint/Private/SceneStateBlueprintBindingUtils.cpp
@@ -31,6 +31,24 @@ FText GetParametersSectionLabel()
 
 } // UE::SceneState::Graph::Private
 
+USceneStateMachineStateNode* FindStateMatchingId(const USceneStateBlueprint& InBlueprint, const FGuid& InStructId)
+{
+	USceneStateMachineStateNode* FoundNode = nullptr;
+
+	VisitNodes(InBlueprint.StateMachineGraphs,
+		[&FoundNode, &InStructId](USceneStateMachineNode* InNode, EIterationResult& IterationResult)
+		{
+			USceneStateMachineStateNode* StateNode = Cast<USceneStateMachineStateNode>(InNode);
+			if (StateNode && StateNode->ParametersId == InStructId)
+			{
+				FoundNode = StateNode;
+				IterationResult = EIterationResult::Break;
+			}
+		});
+
+	return FoundNode;
+}
+
 void GetStateBindingDescs(const USceneStateMachineStateNode* InStateNode, TArray<TInstancedStruct<FSceneStateBindingDesc>>& OutBindingDescs, FString InCategory)
 {
 	if (!InStateNode)
@@ -71,7 +89,12 @@ void GetStateBindingDescs(const USceneStateMachineStateNode* InStateNode, TArray
 		InCategory += TEXT("|");
 	}
 	InCategory += InStateNode->GetNodeName().ToString();
-	GetStateBindingDescs(StateMachineGraph->GetParentStateNode(), OutBindingDescs, InCategory);
+
+	if (USceneStateMachineStateNode* ParentStateNode = StateMachineGraph->GetParentStateNode())
+	{
+		OutBindingDescs.Emplace(TInstancedStruct<FSceneStateBindingDesc>::Make(CreateBindingDesc(*ParentStateNode)));
+		GetStateBindingDescs(ParentStateNode, OutBindingDescs, InCategory);
+	}
 }
 
 FSceneStateBindingDesc CreateBindingDesc(const USceneStateBlueprint& InBlueprint)
@@ -94,6 +117,16 @@ FSceneStateBindingDesc CreateBindingDesc(const USceneStateMachineGraph& InGraph)
 	return BindingDesc;
 }
 
+FSceneStateBindingDesc CreateBindingDesc(const USceneStateMachineStateNode& InStateNode)
+{
+	FSceneStateBindingDesc BindingDesc;
+	BindingDesc.Name = TEXT("State Parameters");
+	BindingDesc.ID = InStateNode.ParametersId;
+	BindingDesc.Struct = InStateNode.Parameters.GetPropertyBagStruct();
+	BindingDesc.Section = Private::GetParametersSectionLabel();
+	return BindingDesc;
+}
+
 FSceneStateBindingDesc CreateBindingDesc(const USceneStateMachineTransitionNode& InTransitionNode)
 {
 	FSceneStateBindingDesc BindingDesc;
@@ -156,7 +189,13 @@ void GetStateMachineBindingDescs(const USceneStateBlueprint& InBlueprint, const
 {
 	// If this State Machine is under a parent state node, add the event handlers of the state (recursively)
 	USceneStateMachineStateNode* ParentStateNode = InGraph.GetTypedOuter<USceneStateMachineStateNode>(); 
-	GetStateBindingDescs(ParentStateNode, OutBindingDescs, FString());
+
+	// Add the state binding desc
+	if (ParentStateNode)
+	{
+		OutBindingDescs.Emplace(TInstancedStruct<FSceneStateBindingDesc>::Make(CreateBindingDesc(*ParentStateNode)));
+		GetStateBindingDescs(ParentStateNode, OutBindingDescs, FString());
+	}
 }
 
 USceneStateMachineTaskNode* FindTaskNodeContainingId(const USceneStateBlueprint& InBlueprint, const FGuid& InStructId)
@@ -183,6 +222,7 @@ void GetTaskBindingDescs(const USceneStateBlueprint& InBlueprint, const USceneSt
 	// Gather the Binding Descs starting from the directly connected State node.
 	if (USceneStateMachineStateNode* const StateNode = USceneStateMachineGraphSchema::FindConnectedStateNode(&InTaskNode))
 	{
+		OutBindingDescs.Emplace(TInstancedStruct<FSceneStateBindingDesc>::Make(CreateBindingDesc(*StateNode)));
 		GetStateBindingDescs(StateNode, OutBindingDescs, FString());
 	}
 	// If the task isn't connected to a State, gather Binding Descs starting from the outer State Machine Graph
@@ -192,8 +232,12 @@ void GetTaskBindingDescs(const USceneStateBlueprint& InBlueprint, const USceneSt
 		OutBindingDescs.Emplace(TInstancedStruct<FSceneStateBindingDesc>::Make(CreateBindingDesc(*StateMachineGraph)));
 
 		// If the Parent State Machine is under a parent outer state node, add the binding descs of that state (recursively)
-		USceneStateMachineStateNode* ParentStateNode = StateMachineGraph->GetTypedOuter<USceneStateMachineStateNode>(); 
-		GetStateBindingDescs(ParentStateNode, OutBindingDescs, FString());
+		USceneStateMachineStateNode* ParentStateNode = StateMachineGraph->GetTypedOuter<USceneStateMachineStateNode>();
+		if (ParentStateNode)
+		{
+			OutBindingDescs.Emplace(TInstancedStruct<FSceneStateBindingDesc>::Make(CreateBindingDesc(*ParentStateNode)));
+			GetStateBindingDescs(ParentStateNode, OutBindingDescs, FString());
+		}
 	}
 }
 
@@ -246,8 +290,12 @@ void GetTransitionBindingDescs(const USceneStateBlueprint& InBlueprint, const US
 		OutBindingDescs.Emplace(TInstancedStruct<FSceneStateBindingDesc>::Make(CreateBindingDesc(*StateMachineGraph)));
 
 		// If the Parent State Machine is under a parent outer state node, add the binding descs of that state (recursively)
-		USceneStateMachineStateNode* ParentStateNode = StateMachineGraph->GetTypedOuter<USceneStateMachineStateNode>(); 
-		GetStateBindingDescs(ParentStateNode, OutBindingDescs, FString());
+		USceneStateMachineStateNode* ParentStateNode = StateMachineGraph->GetTypedOuter<USceneStateMachineStateNode>();
+		if (ParentStateNode)
+		{
+			OutBindingDescs.Emplace(TInstancedStruct<FSceneStateBindingDesc>::Make(CreateBindingDesc(*ParentStateNode)));
+			GetStateBindingDescs(ParentStateNode, OutBindingDescs, FString());
+		}
 	}
 }
 
@@ -312,25 +360,42 @@ bool FindBindingDescById(const USceneStateBlueprint& InBlueprint, const FGuid& I
 			}
 			else if (USceneStateMachineStateNode* StateNode = Cast<USceneStateMachineStateNode>(InNode))
 			{
-				for (const FSceneStateEventHandler& EventHandler : StateNode->GetEventHandlers())
+				if (StateNode->ParametersId == InStructId)
 				{
-					if (EventHandler.GetHandlerId() != InStructId)
-					{
-						continue;
-					}
+					// Found but not yet valid.
+					Result = EResult::FoundInvalid;
+					IterationResult = EIterationResult::Break;
 
-					if (const USceneStateEventSchemaObject* EventSchema = EventHandler.GetEventSchemaHandle().GetEventSchema())
+					FSceneStateBindingDesc BindingDesc = CreateBindingDesc(*StateNode);
+					if (BindingDesc.Struct)
 					{
-						OutBindingDesc = TInstancedStruct<FSceneStateBindingDesc>::Make();
-
-						FSceneStateBindingDesc& StructDesc = OutBindingDesc.GetMutable();
-						StructDesc.ID = EventHandler.GetHandlerId();
-						StructDesc.Name = EventSchema->Name;
-						StructDesc.Struct = EventSchema->Struct;
-
-						IterationResult = EIterationResult::Break;
+						// Found valid binding
+						OutBindingDesc = TInstancedStruct<FSceneStateBindingDesc>::Make(MoveTemp(BindingDesc));
 						Result = EResult::FoundValid;
-						return;
+					}
+				}
+				else
+				{
+					for (const FSceneStateEventHandler& EventHandler : StateNode->GetEventHandlers())
+					{
+						if (EventHandler.GetHandlerId() != InStructId)
+						{
+							continue;
+						}
+
+						if (const USceneStateEventSchemaObject* EventSchema = EventHandler.GetEventSchemaHandle().GetEventSchema())
+						{
+							OutBindingDesc = TInstancedStruct<FSceneStateBindingDesc>::Make();
+
+							FSceneStateBindingDesc& StructDesc = OutBindingDesc.GetMutable();
+							StructDesc.ID = EventHandler.GetHandlerId();
+							StructDesc.Name = EventSchema->Name;
+							StructDesc.Struct = EventSchema->Struct;
+
+							IterationResult = EIterationResult::Break;
+							Result = EResult::FoundValid;
+							return;
+						}
 					}
 				}
 			}
@@ -400,20 +465,29 @@ bool FindDataViewById(const USceneStateBlueprint& InBlueprint, const FGuid& InSt
 					IterationResult = EIterationResult::Break;
 				}
 			}
-			// When finding the Data View for Blueprints, search all States for the Event Handler matching the Struct Id
+			// When finding the Data View for Blueprints, search all States for the Parameters or Event Handler matching the Struct Id
 			else if (USceneStateMachineStateNode* StateNode = Cast<USceneStateMachineStateNode>(InNode))
 			{
-				for (const FSceneStateEventHandler& EventHandler : StateNode->GetEventHandlers())
+				if (StateNode->ParametersId == InStructId)
 				{
-					if (EventHandler.GetHandlerId() != InStructId)
-					{
-						continue;
-					}
-
-					OutDataView = EventHandler.GetEventSchemaHandle().GetDefaultDataView();
+					OutDataView = StateNode->Parameters.GetMutableValue();
 					bDataViewFound = true;
 					IterationResult = EIterationResult::Break;
-					return; // stop execution
+				}
+				else
+				{
+					for (const FSceneStateEventHandler& EventHandler : StateNode->GetEventHandlers())
+					{
+						if (EventHandler.GetHandlerId() != InStructId)
+						{
+							continue;
+						}
+
+						OutDataView = EventHandler.GetEventSchemaHandle().GetDefaultDataView();
+						bDataViewFound = true;
+						IterationResult = EIterationResult::Break;
+						return; // stop execution
+					}
 				}
 			}
 		});
diff --git a/Engine/Plugins/VirtualProduction/SceneState/Source/SceneStateBlueprint/Private/SceneStateBlueprintBindingUtils.h b/Engine/Plugins/VirtualProduction/SceneState/Source/SceneStateBlueprint/Private/SceneStateBlueprintBindingUtils.h
--- a/Engine/Plugins/VirtualProduction/SceneState/Source/SceneStateBlueprint/Private/SceneStateBlueprintBindingUtils.h
+++ b/Engine/Plugins/VirtualProduction/SceneState/Source/SceneStateBlueprint/Private/SceneStateBlueprintBindingUtils.h
@@ -19,6 +19,14 @@ template<typename BaseStructT> struct TInstancedStruct;
 
 namespace UE::SceneState::Graph
 {
+	/**
+	 * Finds the state node with id that matches the given struct id.
+	 * @param InBlueprint the blueprint to look into
+	 * @param InStructId the struct id to look for
+	 * @return the found state node matching the given struct id, or null if not found
+	 */
+	USceneStateMachineStateNode* FindStateMatchingId(const USceneStateBlueprint& InBlueprint, const FGuid& InStructId);
+
 	/**
 	 * Gathers the Binding Descs for a State Node
 	 * @param InStateNode the state to find binding descs for
@@ -33,6 +41,9 @@ namespace UE::SceneState::Graph
 	/** Creates the Binding Struct Desc for the Parameters in the given State Machine Graph */
 	FSceneStateBindingDesc CreateBindingDesc(const USceneStateMachineGraph& InGraph);
 
+	/** Creates the Binding Struct Desc for the Parameters in the given State Node */
+	FSceneStateBindingDesc CreateBindingDesc(const USceneStateMachineStateNode& InStateNode);
+
 	/** Creates the Binding Struct Desc for the Parameters in the given Transition Node */
 	FSceneStateBindingDesc CreateBindingDesc(const USceneStateMachineTransitionNode& InTransitionNode);
 
diff --git a/Engine/Plugins/VirtualProduction/SceneState/Source/SceneStateBlueprint/Public/SceneStateBlueprint.h b/Engine/Plugins/VirtualProduction/SceneState/Source/SceneStateBlueprint/Public/SceneStateBlueprint.h
--- a/Engine/Plugins/VirtualProduction/SceneState/Source/SceneStateBlueprint/Public/SceneStateBlueprint.h
+++ b/Engine/Plugins/VirtualProduction/SceneState/Source/SceneStateBlueprint/Public/SceneStateBlueprint.h
@@ -109,8 +109,11 @@ private:
 	/** Called when a state machine graph's parameters have changed */
 	void OnGraphParametersChanged(USceneStateMachineGraph* InGraph);
 
-	/** Called when a state machine graph's parameters have changed. Used to fix the path to parameters that have possibly been renamed */
-	void UpdateGraphParametersBindings(FPropertyBindingPath& InPath, USceneStateMachineGraph* InGraph);
+	/** Called when a state's parameters have changed */
+	void OnStateParametersChanged(USceneStateMachineStateNode* InStateNode);
+
+	/** Called when parameters have changed. Used to fix the path to parameters that have possibly been renamed */
+	void UpdateParametersBindings(FPropertyBindingPath& InPath, const FGuid& InParametersId, const FInstancedPropertyBag& InParameters);
 
 	/** Scene State schema class to use for this blueprint */
 	UPROPERTY(VisibleAnywhere, Category="Scene State", AssetRegistrySearchable)
@@ -125,6 +128,9 @@ private:
 
 	/** Handle to the delegate when a state machine graph parameters have changed */
 	FDelegateHandle OnGraphParametersChangedHandle;
+
+	/** Handle to the delegate when a state parameters have changed */
+	FDelegateHandle OnStateParametersChangedHandle;
 };
 
 #undef UE_API
diff --git a/Engine/Plugins/VirtualProduction/SceneState/Source/SceneStateBlueprintEditor/Private/Compiler/SceneStateBindingCompiler.cpp b/Engine/Plugins/VirtualProduction/SceneState/Source/SceneStateBlueprintEditor/Private/Compiler/SceneStateBindingCompiler.cpp
--- a/Engine/Plugins/VirtualProduction/SceneState/Source/SceneStateBlueprintEditor/Private/Compiler/SceneStateBindingCompiler.cpp
+++ b/Engine/Plugins/VirtualProduction/SceneState/Source/SceneStateBlueprintEditor/Private/Compiler/SceneStateBindingCompiler.cpp
@@ -82,6 +82,7 @@ void FBindingCompiler::Compile()
 	BindingFunctionCompiler.Compile();
 
 	AddRootBindingDesc();
+	AddStateBindingDescs();
 	AddStateMachineBindingDescs();
 	AddTransitionBindingDescs();
 	AddTaskBindingDescs();
@@ -122,6 +123,28 @@ void FBindingCompiler::AddRootBindingDesc()
 	AddBindingDesc(MoveTemp(BindingDesc));
 }
 
+void FBindingCompiler::AddStateBindingDescs()
+{
+	for (const TPair<FGuid, uint16>& Pair : TemplateData->StateParameterIdToIndex)
+	{
+		const uint16 StateIndex = Pair.Value;
+
+		const FSceneState& State = TemplateData->States[StateIndex];
+		if (!State.Parameters.IsValid())
+		{
+			continue;
+		}
+
+		FSceneStateBindingDesc BindingDesc;
+		BindingDesc.ID = Pair.Key;
+		BindingDesc.Name = TEXT("State Parameters");
+		BindingDesc.Struct = State.Parameters.GetPropertyBagStruct();
+		BindingDesc.DataHandle = FSceneStateBindingDataHandle(ESceneStateDataType::State, StateIndex);
+
+		AddBindingDesc(MoveTemp(BindingDesc));
+	}
+}
+
 void FBindingCompiler::AddStateMachineBindingDescs()
 {
 	for (const TPair<FGuid, uint16>& Pair : TemplateData->StateMachineIdToIndex)
@@ -521,6 +544,11 @@ void FBindingCompiler::OnBindingsBatchCompiled(FPropertyBindingIndex16 InBinding
 		FSceneStateFunction& Function = TemplateData->Functions[DataIndex].Get<FSceneStateFunction>();
 		Function.BindingsBatch = InBindingsBatch;
 	}
+	else if (DataType == ESceneStateDataType::State)
+	{
+		FSceneState& State = TemplateData->States[DataIndex];
+		State.BindingsBatch = InBindingsBatch;
+	}
 	else
 	{
 		// No other Data Types are supported as Targets at the moment
@@ -602,6 +630,15 @@ FSceneStateBindingDataHandle FBindingCompiler::GetDataHandleById(const FGuid& In
 		}
 	}
 
+	// States
+	for (const TPair<FGuid, uint16>& Pair : TemplateData->StateParameterIdToIndex)
+	{
+		if (Pair.Key == InStructId)
+		{
+			return FSceneStateBindingDataHandle(ESceneStateDataType::State, Pair.Value);
+		}
+	}
+
 	return FSceneStateBindingDataHandle();
 }
 
diff --git a/Engine/Plugins/VirtualProduction/SceneState/Source/SceneStateBlueprintEditor/Private/Compiler/SceneStateBindingCompiler.h b/Engine/Plugins/VirtualProduction/SceneState/Source/SceneStateBlueprintEditor/Private/Compiler/SceneStateBindingCompiler.h
--- a/Engine/Plugins/VirtualProduction/SceneState/Source/SceneStateBlueprintEditor/Private/Compiler/SceneStateBindingCompiler.h
+++ b/Engine/Plugins/VirtualProduction/SceneState/Source/SceneStateBlueprintEditor/Private/Compiler/SceneStateBindingCompiler.h
@@ -42,6 +42,7 @@ private:
 	void AddBindingDesc(FSceneStateBindingDesc&& InBindingDesc);
 
 	void AddRootBindingDesc();
+	void AddStateBindingDescs();
 	void AddStateMachineBindingDescs();
 	void AddTransitionBindingDescs();
 	void AddTaskBindingDescs();
diff --git a/Engine/Plugins/VirtualProduction/SceneState/Source/SceneStateBlueprintEditor/Private/DetailsView/SceneStateMachineStateNodeCustomization.cpp b/Engine/Plugins/VirtualProduction/SceneState/Source/SceneStateBlueprintEditor/Private/DetailsView/SceneStateMachineStateNodeCustomization.cpp
new file mode 100644
--- /dev/null
+++ b/Engine/Plugins/VirtualProduction/SceneState/Source/SceneStateBlueprintEditor/Private/DetailsView/SceneStateMachineStateNodeCustomization.cpp
@@ -0,0 +1,55 @@
+// Copyright Epic Games, Inc. All Rights Reserved.
+
+#include "SceneStateMachineStateNodeCustomization.h"
+#include "DetailCategoryBuilder.h"
+#include "DetailLayoutBuilder.h"
+#include "Nodes/SceneStateMachineStateNode.h"
+#include "SceneStateBlueprintEditorUtils.h"
+#include "SceneStateParameterDetails.h"
+
+namespace UE::SceneState::Editor
+{
+
+void FStateMachineStateNodeCustomization::CustomizeDetails(IDetailLayoutBuilder& InDetailBuilder)
+{
+	const TArray<TWeakObjectPtr<USceneStateMachineStateNode>> StateNodes = InDetailBuilder.GetObjectsOfTypeBeingCustomized<USceneStateMachineStateNode>();
+
+	const TSharedRef<IPropertyHandle> ParametersIdHandle = InDetailBuilder.GetProperty(GET_MEMBER_NAME_CHECKED(USceneStateMachineStateNode, ParametersId));
+	const TSharedRef<IPropertyHandle> ParametersHandle = InDetailBuilder.GetProperty(GET_MEMBER_NAME_CHECKED(USceneStateMachineStateNode, Parameters));
+	const TSharedRef<IPropertyHandle> EventHandlersHandle = InDetailBuilder.GetProperty(GET_MEMBER_NAME_CHECKED(USceneStateMachineStateNode, EventHandlers));
+
+	ParametersIdHandle->MarkHiddenByCustomization();
+	ParametersHandle->MarkHiddenByCustomization();
+	EventHandlersHandle->MarkHiddenByCustomization();
+
+	// Listen to Parameter Changes to broadcast change
+	ParametersHandle->SetOnPropertyValueChanged(FSimpleDelegate::CreateSPLambda(this,
+		[StateNodes]()
+		{
+			for (TWeakObjectPtr<USceneStateMachineStateNode> StateNodeWeak : StateNodes)
+			{
+				if (USceneStateMachineStateNode* StateNode = StateNodeWeak.Get())
+				{
+					StateNode->NotifyParametersChanged();
+				}
+			}
+		}));
+
+	FGuid ParametersId;
+	GetGuid(ParametersIdHandle, ParametersId);
+
+	// Parameters Category
+	IDetailCategoryBuilder& ParametersCategory = InDetailBuilder.EditCategory(TEXT("Parameters"));
+	ParametersCategory.HeaderContent(FParameterDetails::BuildHeader(InDetailBuilder, ParametersHandle), /*bWholeRowContent*/true);
+
+	ParametersCategory.AddCustomBuilder(MakeShared<FParameterDetails>(ParametersHandle
+		, InDetailBuilder.GetPropertyUtilities()
+		, ParametersId
+		, /*bFixedLayout*/false));
+
+	// Event Handlers Category
+	IDetailCategoryBuilder& EventsCategory = InDetailBuilder.EditCategory(TEXT("Events"));
+	EventsCategory.AddProperty(EventHandlersHandle);
+}
+
+} // UE::SceneState::Editor
diff --git a/Engine/Plugins/VirtualProduction/SceneState/Source/SceneStateBlueprintEditor/Private/DetailsView/SceneStateMachineStateNodeCustomization.h b/Engine/Plugins/VirtualProduction/SceneState/Source/SceneStateBlueprintEditor/Private/DetailsView/SceneStateMachineStateNodeCustomization.h
new file mode 100644
--- /dev/null
+++ b/Engine/Plugins/VirtualProduction/SceneState/Source/SceneStateBlueprintEditor/Private/DetailsView/SceneStateMachineStateNodeCustomization.h
@@ -0,0 +1,23 @@
+// Copyright Epic Games, Inc. All Rights Reserved.
+
+#pragma once
+
+#include "IDetailCustomization.h"
+
+namespace UE::SceneState::Editor
+{
+
+class FStateMachineStateNodeCustomization : public IDetailCustomization
+{
+public:
+	static TSharedRef<IDetailCustomization> MakeInstance()
+	{
+		return MakeShared<FStateMachineStateNodeCustomization>();
+	}
+
+	//~ Begin IDetailCustomization
+	virtual void CustomizeDetails(IDetailLayoutBuilder& InDetailBuilder) override;
+	//~ End IDetailCustomization
+};
+
+} // UE::SceneState::Editor
diff --git a/Engine/Plugins/VirtualProduction/SceneState/Source/SceneStateBlueprintEditor/Private/SceneStateBlueprintEditorModule.cpp b/Engine/Plugins/VirtualProduction/SceneState/Source/SceneStateBlueprintEditor/Private/SceneStateBlueprintEditorModule.cpp
--- a/Engine/Plugins/VirtualProduction/SceneState/Source/SceneStateBlueprintEditor/Private/SceneStateBlueprintEditorModule.cpp
+++ b/Engine/Plugins/VirtualProduction/SceneState/Source/SceneStateBlueprintEditor/Private/SceneStateBlueprintEditorModule.cpp
@@ -8,12 +8,14 @@
 #include "DetailsView/SceneStateBlueprintPropertyReferenceCustomization.h"
 #include "DetailsView/SceneStateBlueprintableTaskInstanceCustomization.h"
 #include "DetailsView/SceneStateMachineGraphCustomization.h"
+#include "DetailsView/SceneStateMachineStateNodeCustomization.h"
 #include "DetailsView/SceneStateMachineTaskInstanceCustomization.h"
 #include "DetailsView/SceneStateMachineTransitionNodeCustomization.h"
 #include "DetailsView/SceneStatePlayerTaskInstanceCustomization.h"
 #include "Kismet2/KismetEditorUtilities.h"
 #include "KismetCompiler.h"
 #include "KismetCompilerModule.h"
+#include "Nodes/SceneStateMachineStateNode.h"
 #include "Nodes/SceneStateMachineTaskNode.h"
 #include "Nodes/SceneStateMachineTransitionNode.h"
 #include "PropertyEditorModule.h"
@@ -180,6 +182,9 @@ void FBlueprintEditorModule::RegisterDetailCustomizations()
 	PropertyEditorModule.RegisterCustomClassLayout(CustomizedClasses.Add_GetRef(USceneStateMachineGraph::StaticClass()->GetFName())
 		, FOnGetDetailCustomizationInstance::CreateStatic(&FStateMachineGraphCustomization::MakeInstance));
 
+	PropertyEditorModule.RegisterCustomClassLayout(CustomizedClasses.Add_GetRef(USceneStateMachineStateNode::StaticClass()->GetFName())
+		, FOnGetDetailCustomizationInstance::CreateStatic(&FStateMachineStateNodeCustomization::MakeInstance));
+
 	PropertyEditorModule.RegisterCustomClassLayout(CustomizedClasses.Add_GetRef(USceneStateMachineTransitionNode::StaticClass()->GetFName())
 		, FOnGetDetailCustomizationInstance::CreateStatic(&FStateMachineTransitionNodeCustomization::MakeInstance));
 }
diff --git a/Engine/Plugins/VirtualProduction/SceneState/Source/SceneStateMachineEditor/Private/SceneStateMachineCompiler.cpp b/Engine/Plugins/VirtualProduction/SceneState/Source/SceneStateMachineEditor/Private/SceneStateMachineCompiler.cpp
--- a/Engine/Plugins/VirtualProduction/SceneState/Source/SceneStateMachineEditor/Private/SceneStateMachineCompiler.cpp
+++ b/Engine/Plugins/VirtualProduction/SceneState/Source/SceneStateMachineEditor/Private/SceneStateMachineCompiler.cpp
@@ -125,6 +125,12 @@ FSceneStateMachine FStateMachineCompiler::Compile()
 	// Upgrade the map to absolute indices before baking it to the generated class' map
 	ToAbsoluteIndexMap(StateNodeIndexMap, StateMachine.StateRange.Index);
 	TemplateData->StateNodeToIndex.Append(MoveTemp(StateNodeIndexMap));
+	StateNodeIndexMap.Reset();
+
+	// Upgrade the map to absolute indices before baking it to the generated class' map
+	ToAbsoluteIndexMap(StateParameterIdIndexMap, StateMachine.StateRange.Index);
+	TemplateData->StateParameterIdToIndex.Append(MoveTemp(StateParameterIdIndexMap));
+	StateParameterIdIndexMap.Reset();
 
 	StateMachine.ConduitRange.Index = TemplateData->Conduits.Num();
 	StateMachine.ConduitRange.Count = Conduits.Num();
@@ -156,11 +162,13 @@ void FStateMachineCompiler::CompileState(const USceneStateMachineStateNode* InSt
 	check(StateIndex == StateMetadataIndex);
 
 	FSceneState& NewState = States[StateIndex];
+	NewState.Parameters = InStateNode->Parameters;
 
 	FSceneStateMetadata& NewStateMetadata = StateMetadata[StateIndex];
 	NewStateMetadata.StateName = InStateNode->GetNodeName().ToString();
 
 	StateNodeIndexMap.Add(InStateNode, StateIndex);
+	StateParameterIdIndexMap.Add(InStateNode->ParametersId, StateIndex);
 
 	CompileSubStateMachines(NewState, InStateNode);
 	CompileTasks(NewState, InStateNode);
diff --git a/Engine/Plugins/VirtualProduction/SceneState/Source/SceneStateMachineEditor/Public/SceneStateMachineCompiler.h b/Engine/Plugins/VirtualProduction/SceneState/Source/SceneStateMachineEditor/Public/SceneStateMachineCompiler.h
--- a/Engine/Plugins/VirtualProduction/SceneState/Source/SceneStateMachineEditor/Public/SceneStateMachineCompiler.h
+++ b/Engine/Plugins/VirtualProduction/SceneState/Source/SceneStateMachineEditor/Public/SceneStateMachineCompiler.h
@@ -86,6 +86,9 @@ private:
 	 */
 	TMap<FObjectKey, uint16> StateNodeIndexMap;
 
+	/** Map of a State's parameter id to the index of the states array below */
+	TMap<FGuid, uint16> StateParameterIdIndexMap;
+
 	/** Compiled States */
 	TArray<FSceneState> States;
 
diff --git a/Engine/Plugins/VirtualProduction/SceneState/Source/SceneStateMachineGraph/Private/Nodes/SceneStateMachineStateNode.cpp b/Engine/Plugins/VirtualProduction/SceneState/Source/SceneStateMachineGraph/Private/Nodes/SceneStateMachineStateNode.cpp
--- a/Engine/Plugins/VirtualProduction/SceneState/Source/SceneStateMachineGraph/Private/Nodes/SceneStateMachineStateNode.cpp
+++ b/Engine/Plugins/VirtualProduction/SceneState/Source/SceneStateMachineGraph/Private/Nodes/SceneStateMachineStateNode.cpp
@@ -4,9 +4,12 @@
 #include "Kismet2/BlueprintEditorUtils.h"
 #include "Kismet2/Kismet2NameValidators.h"
 #include "Nodes/SceneStateMachineTaskNode.h"
+#include "SceneStateBindingUtils.h"
 #include "SceneStateMachineGraph.h"
 #include "SceneStateMachineGraphSchema.h"
 
+USceneStateMachineStateNode::FOnParametersChanged USceneStateMachineStateNode::OnParametersChangedDelegate;
+
 USceneStateMachineStateNode::USceneStateMachineStateNode()
 {
 	NodeName = TEXT("State");
@@ -15,6 +18,16 @@ USceneStateMachineStateNode::USceneStateMachineStateNode()
 	bCanRenameNode = true;
 }
 
+TMulticastDelegate<void(USceneStateMachineStateNode*)>::RegistrationType& USceneStateMachineStateNode::OnParametersChanged()
+{
+	return OnParametersChangedDelegate;
+}
+
+void USceneStateMachineStateNode::NotifyParametersChanged()
+{
+	OnParametersChangedDelegate.Broadcast(this);
+}
+
 TOptional<uint16> USceneStateMachineStateNode::GetCompiledIndex() const
 {
 	return CompiledIndex;
@@ -84,6 +97,16 @@ void USceneStateMachineStateNode::PostPlacedNewNode()
 	GenerateNewNodeName();
 }
 
+void USceneStateMachineStateNode::PostInitProperties()
+{
+	Super::PostInitProperties();
+
+	if (!IsTemplate())
+	{
+		ParametersId = FGuid::NewGuid();
+	}
+}
+
 void USceneStateMachineStateNode::PostLoad()
 {
 	Super::PostLoad();
@@ -99,3 +122,22 @@ void USceneStateMachineStateNode::PostLoad()
 	}
 	MainGraph = nullptr;
 }
+
+void USceneStateMachineStateNode::PostDuplicate(bool bInDuplicateForPIE)
+{
+	Super::PostDuplicate(bInDuplicateForPIE);
+	GenerateNewParametersId();
+}
+
+void USceneStateMachineStateNode::PostEditImport()
+{
+	Super::PostEditImport();
+	GenerateNewParametersId();
+}
+
+void USceneStateMachineStateNode::GenerateNewParametersId()
+{
+	const FGuid OldParametersId = ParametersId;
+	ParametersId = FGuid::NewGuid();
+	UE::SceneState::HandleStructIdChanged(*this, OldParametersId, ParametersId);
+}
diff --git a/Engine/Plugins/VirtualProduction/SceneState/Source/SceneStateMachineGraph/Public/Nodes/SceneStateMachineStateNode.h b/Engine/Plugins/VirtualProduction/SceneState/Source/SceneStateMachineGraph/Public/Nodes/SceneStateMachineStateNode.h
--- a/Engine/Plugins/VirtualProduction/SceneState/Source/SceneStateMachineGraph/Public/Nodes/SceneStateMachineStateNode.h
+++ b/Engine/Plugins/VirtualProduction/SceneState/Source/SceneStateMachineGraph/Public/Nodes/SceneStateMachineStateNode.h
@@ -5,6 +5,7 @@
 #include "ISceneStateEventHandlerProvider.h"
 #include "SceneStateEventHandler.h"
 #include "SceneStateMachineNode.h"
+#include "StructUtils/PropertyBag.h"
 #include "SceneStateMachineStateNode.generated.h"
 
 #define UE_API SCENESTATEMACHINEGRAPH_API
@@ -15,8 +16,14 @@ class USceneStateMachineStateNode : public USceneStateMachineNode, public IScene
 	GENERATED_BODY()
 
 public:
+	using FOnParametersChanged = TMulticastDelegate<void(USceneStateMachineStateNode*)>;
+
 	USceneStateMachineStateNode();
 
+	UE_INTERNAL UE_API static FOnParametersChanged::RegistrationType& OnParametersChanged();
+
+	UE_INTERNAL UE_API void NotifyParametersChanged();
+
 	UEdGraphPin* GetTaskPin() const
 	{
 		return Pins[2];
@@ -50,8 +57,22 @@ public:
 	//~ End UEdGraphNode
 
 	//~ Begin UObject
+	virtual void PostInitProperties() override;
 	virtual void PostLoad() override;
-	//~ End Uobject
+	virtual void PostDuplicate(bool bInDuplicateForPIE) override;
+	virtual void PostEditImport() override;
+	//~ End UObject
+
+	/** Called to set a new unique id for parameters (e.g. after duplicating) */
+	void GenerateNewParametersId();
+
+	/** Identifier for the Parameters Struct Id */
+	UPROPERTY(VisibleAnywhere, Category="Parameters")
+	FGuid ParametersId;
+
+	/** State parameters */
+	UPROPERTY(EditAnywhere, Category="Parameters")
+	FInstancedPropertyBag Parameters;
 
 	/** Deprecated: Graphs are now managed in the Node Base class */
 	UPROPERTY()
@@ -64,6 +85,8 @@ private:
 	/** Index to the state this state node compiled to last. */
 	UPROPERTY()
 	TOptional<uint16> CompiledIndex;
+
+	static FOnParametersChanged OnParametersChangedDelegate;
 };
 
 #undef UE_API
