Motion Design Scene State: changed schema to only check for task type metadata in the task itself, not its super structs.
+ Added gameplay task type metadata
+ Comments on metadata for clarity

#Synced-CL 50867481

diff --git a/Engine/Plugins/VirtualProduction/AvalancheSceneState/Source/AvalancheSceneState/Private/AvaSceneStateSchema.cpp b/Engine/Plugins/VirtualProduction/AvalancheSceneState/Source/AvalancheSceneState/Private/AvaSceneStateSchema.cpp
--- a/Engine/Plugins/VirtualProduction/AvalancheSceneState/Source/AvalancheSceneState/Private/AvaSceneStateSchema.cpp
+++ b/Engine/Plugins/VirtualProduction/AvalancheSceneState/Source/AvalancheSceneState/Private/AvaSceneStateSchema.cpp
@@ -12,7 +12,7 @@
 bool UAvaSceneStateSchema::OnIsTaskStructAllowed(TSubScriptStructOf<FSceneStateTask> InTaskStruct) const
 {
 	// Allow Motion Design and Core tasks
-	constexpr bool bIncludeSuperStructs = true;
+	constexpr bool bIncludeSuperStructs = false;
 	return UE::SceneState::Editor::StructHasAnyMetaData(InTaskStruct, 
 		{
 			UE::SceneState::Editor::MD_CoreTask,
diff --git a/Engine/Plugins/VirtualProduction/AvalancheSceneState/Source/AvalancheSceneState/Public/AvaSceneStateSchema.h b/Engine/Plugins/VirtualProduction/AvalancheSceneState/Source/AvalancheSceneState/Public/AvaSceneStateSchema.h
--- a/Engine/Plugins/VirtualProduction/AvalancheSceneState/Source/AvalancheSceneState/Public/AvaSceneStateSchema.h
+++ b/Engine/Plugins/VirtualProduction/AvalancheSceneState/Source/AvalancheSceneState/Public/AvaSceneStateSchema.h
@@ -10,7 +10,10 @@
 #if WITH_EDITOR
 namespace UE::AvaSceneState::Editor
 {
-	/** Metadata to identify a task as a Motion Design task */
+	/** 
+	 * Metadata to identify a task as a Motion Design task 
+	 * Derived tasks will NOT be recognized as 'MotionDesignTask' unless they also specify this metadata.
+	 */
 	const FLazyName MD_MotionDesignTask(TEXT("MotionDesignTask"));
 }
 #endif
diff --git a/Engine/Plugins/VirtualProduction/SceneState/Source/SceneState/Public/SceneStateUtils.h b/Engine/Plugins/VirtualProduction/SceneState/Source/SceneState/Public/SceneStateUtils.h
--- a/Engine/Plugins/VirtualProduction/SceneState/Source/SceneState/Public/SceneStateUtils.h
+++ b/Engine/Plugins/VirtualProduction/SceneState/Source/SceneState/Public/SceneStateUtils.h
@@ -104,10 +104,10 @@ namespace UE::SceneState
 namespace Editor
 {
 	/** Returns true if the struct or its super structs contains the given metadata. */
-	UE_API bool StructHasMetaData(const UStruct* InStruct, FName InMetaData, bool bInIncludeSuperStructs = true);
+	UE_API bool StructHasMetaData(const UStruct* InStruct, FName InMetaData, bool bInIncludeSuperStructs);
 
 	/** Returns true if the struct or its super structs contains any of the given metadata. */
-	UE_API bool StructHasAnyMetaData(const UStruct* InStruct, TConstArrayView<FName> InMetaData, bool bInIncludeSuperStructs = true);
+	UE_API bool StructHasAnyMetaData(const UStruct* InStruct, TConstArrayView<FName> InMetaData, bool bInIncludeSuperStructs);
 
 } // UE::SceneState::Editor
 #endif
diff --git a/Engine/Plugins/VirtualProduction/SceneState/Source/SceneState/Public/Tasks/SceneStateTaskMetadata.h b/Engine/Plugins/VirtualProduction/SceneState/Source/SceneState/Public/Tasks/SceneStateTaskMetadata.h
--- a/Engine/Plugins/VirtualProduction/SceneState/Source/SceneState/Public/Tasks/SceneStateTaskMetadata.h
+++ b/Engine/Plugins/VirtualProduction/SceneState/Source/SceneState/Public/Tasks/SceneStateTaskMetadata.h
@@ -10,16 +10,23 @@ namespace UE::SceneState::Editor
 	/** The task should not appear in the list of available tasks (e.g. because it's a base task struct) */
 	const FLazyName MD_Hidden(TEXT("Hidden"));
 
-	/** The task is deemed as 'Core'. All the tasks that come with the Scene State plugin will be marked with this metadata. */
+	/** 
+	 * The task is deemed as 'Core'. All the tasks that come with the Scene State plugin will be marked with this metadata. 
+	 * Derived tasks will NOT be recognized as 'CoreTask' unless they also specify this metadata.
+	 */
 	const FLazyName MD_CoreTask(TEXT("CoreTask"));
 
 	/**
 	 * The task requires the context object to implement GetWorld().
 	 * This is just a convenience metadata to avoid having to call SupportsSchema.
+	 * Derived tasks will also be treated as requiring a context world.
 	 */
 	const FLazyName MD_RequiresContextWorld(TEXT("RequiresContextWorld"));
 
-	/** The task implements OnSupportsSchema function. Tasks without this metadata will not have this editor function called. */
+	/** 
+	 * The task implements OnSupportsSchema function. Only tasks with this metadata will have their 'SupportsSchema' function called. 
+	 * Derived tasks will also have this function called.
+	 */
 	const FLazyName MD_WithSupportsSchema(TEXT("WithSupportsSchema"));
 #endif
 }
diff --git a/Engine/Plugins/VirtualProduction/SceneState/Source/SceneStateGameplay/Private/SceneStateGameplaySchema.cpp b/Engine/Plugins/VirtualProduction/SceneState/Source/SceneStateGameplay/Private/SceneStateGameplaySchema.cpp
--- a/Engine/Plugins/VirtualProduction/SceneState/Source/SceneStateGameplay/Private/SceneStateGameplaySchema.cpp
+++ b/Engine/Plugins/VirtualProduction/SceneState/Source/SceneStateGameplay/Private/SceneStateGameplaySchema.cpp
@@ -12,8 +12,13 @@
 #if WITH_EDITOR
 bool USceneStateGameplaySchema::OnIsTaskStructAllowed(TSubScriptStructOf<FSceneStateTask> InTaskStruct) const
 {
-	constexpr bool bIncludeSuperStructs = true;
-	return UE::SceneState::Editor::StructHasMetaData(InTaskStruct, UE::SceneState::Editor::MD_CoreTask, bIncludeSuperStructs);
+	// Only include core and gameplay tasks.
+	constexpr bool bIncludeSuperStructs = false;
+	return UE::SceneState::Editor::StructHasAnyMetaData(InTaskStruct,
+		{
+			UE::SceneState::Editor::MD_CoreTask,
+			UE::SceneState::Editor::MD_GameplayTask,
+		}, bIncludeSuperStructs);
 }
 #endif
 
diff --git a/Engine/Plugins/VirtualProduction/SceneState/Source/SceneStateGameplay/Public/SceneStateGameplaySchema.h b/Engine/Plugins/VirtualProduction/SceneState/Source/SceneStateGameplay/Public/SceneStateGameplaySchema.h
--- a/Engine/Plugins/VirtualProduction/SceneState/Source/SceneStateGameplay/Public/SceneStateGameplaySchema.h
+++ b/Engine/Plugins/VirtualProduction/SceneState/Source/SceneStateGameplay/Public/SceneStateGameplaySchema.h
@@ -8,6 +8,17 @@
 
 #define UE_API SCENESTATEGAMEPLAY_API
 
+#if WITH_EDITOR
+namespace UE::SceneState::Editor
+{
+	/** 
+	 * Metadata to identify a task as a gameplay task 
+	 * Derived tasks will NOT be recognized as 'GameplayTask' unless they also specify this metadata.
+	 */
+	const FLazyName MD_GameplayTask(TEXT("GameplayTask"));
+}
+#endif
+
 /** Schema for scene states played via the Scene State Component / Actor context */
 UCLASS(MinimalAPI)
 class USceneStateGameplaySchema : public USceneStateSchema
