Motion Design Scene State: made task instances have no base properties to avoid padding waste

#Synced-CL 50950138

diff --git a/Engine/Plugins/VirtualProduction/SceneState/Source/SceneState/Private/SceneState.cpp b/Engine/Plugins/VirtualProduction/SceneState/Source/SceneState/Private/SceneState.cpp
--- a/Engine/Plugins/VirtualProduction/SceneState/Source/SceneState/Private/SceneState.cpp
+++ b/Engine/Plugins/VirtualProduction/SceneState/Source/SceneState/Private/SceneState.cpp
@@ -33,8 +33,8 @@ void FSceneState::Enter(const FSceneStateExecutionContext& InContext) const
 
 	UE_LOG(LogSceneState, Verbose, TEXT("State (%s) receiving enter"), GetStateName(InContext).GetData());
 
+	Instance.Setup(*this);
 	InContext.SetStateStatus(*this, Instance, EExecutionStatus::Running);
-	Instance.InitializeParameters(Parameters);
 	InContext.SetupFunctionInstances(BindingsBatch);
 
 	// Apply Event Handlers before anything else starts so the Event Data becomes available to this state, substate machines and tasks
@@ -54,9 +54,10 @@ void FSceneState::Enter(const FSceneStateExecutionContext& InContext) const
 
 	// Setup each Task
 	InContext.ForEachTask(*this,
-		[&InContext](const FSceneStateTask& InTask, FStructView InTaskInstance)->EIterationResult
+		[&InContext, &Instance](const FSceneStateTask& InTask, uint16 InTaskRelativeIndex)->EIterationResult
 		{
-			InTask.Setup(InContext, InTaskInstance);
+			const FStructView TaskInstance = Instance.GetTaskInstances()[InTaskRelativeIndex];
+			InTask.Setup(InContext, TaskInstance);
 			return EIterationResult::Continue;
 		});
 
@@ -65,7 +66,7 @@ void FSceneState::Enter(const FSceneStateExecutionContext& InContext) const
 		StateMachine.Start(InContext);
 	}
 
-	UpdateActiveTasks(InContext);
+	UpdateActiveTasks(InContext, Instance);
 }
 
 void FSceneState::Tick(const FSceneStateExecutionContext& InContext, float InDeltaSeconds) const
@@ -86,7 +87,7 @@ void FSceneState::Tick(const FSceneStateExecutionContext& InContext, float InDel
 
 	UE_LOG(LogSceneState, VeryVerbose, TEXT("State (%s) receiving tick"), GetStateName(InContext).GetData());
 
-	UpdateActiveTasks(InContext);
+	UpdateActiveTasks(InContext, *Instance);
 
 	for (const FSceneStateMachine& StateMachine : InContext.GetStateMachines(*this))
 	{
@@ -94,9 +95,10 @@ void FSceneState::Tick(const FSceneStateExecutionContext& InContext, float InDel
 	}
 
 	InContext.ForEachTask(*this,
-		[&InContext, Instance, InDeltaSeconds](const FSceneStateTask& InTask, FStructView InTaskInstance)
+		[&InContext, Instance, InDeltaSeconds](const FSceneStateTask& InTask, uint16 InTaskRelativeIndex)
 		{
-			InTask.Tick(InContext, InTaskInstance, InDeltaSeconds);
+			const FStructView TaskInstance = Instance->GetTaskInstances()[InTaskRelativeIndex];
+			InTask.Tick(InContext, TaskInstance, InDeltaSeconds);
 			return EIterationResult::Continue;
 		});
 }
@@ -127,9 +129,10 @@ void FSceneState::Exit(const FSceneStateExecutionContext& InContext) const
 
 	// Stop tasks that are still running
 	InContext.ForEachTask(*this,
-		[&InContext, Instance](const FSceneStateTask& InTask, FStructView InTaskInstance)
+		[&InContext, Instance](const FSceneStateTask& InTask, uint16 InTaskRelativeIndex)
 		{
-			InTask.Stop(InContext, InTaskInstance, ESceneStateTaskStopReason::State);
+			const FStructView TaskInstance = Instance->GetTaskInstances()[InTaskRelativeIndex];
+			InTask.Stop(InContext, TaskInstance, ESceneStateTaskStopReason::State);
 			return EIterationResult::Continue;
 		});
 
@@ -142,47 +145,39 @@ void FSceneState::Exit(const FSceneStateExecutionContext& InContext) const
 	InContext.SetStateStatus(*this, *Instance, EExecutionStatus::Finished);
 
 	InContext.RemoveStateInstance(*this);
-	InContext.RemoveTaskInstanceContainer(*this);
 	InContext.RemoveFunctionInstances(BindingsBatch);
 
 	ResetCapturedEvents(InContext);
 }
 
-void FSceneState::UpdateActiveTasks(const FSceneStateExecutionContext& InContext) const
+const FInstancedPropertyBag& FSceneState::GetParameters() const
 {
-	using namespace UE::SceneState;
+	return Parameters;
+}
 
-	FInstancedStructContainer* TaskInstanceContainer = InContext.FindTaskInstanceContainer(*this);
-	if (!TaskInstanceContainer)
-	{
-		return;
-	}
+const UStruct* FSceneState::GetParametersStruct() const
+{
+	return Parameters.GetPropertyBagStruct();
+}
 
-	auto FindTaskInstance = 
-		[TaskInstanceContainer](uint16 InRelativeIndex)->FSceneStateTaskInstance*
-		{
-			if (TaskInstanceContainer->IsValidIndex(InRelativeIndex))
-			{
-				return (*TaskInstanceContainer)[InRelativeIndex].GetPtr<FSceneStateTaskInstance>();
-			}
-			return nullptr;
-		};
+void FSceneState::UpdateActiveTasks(const FSceneStateExecutionContext& InContext, FSceneStateInstance& InStateInstance) const
+{
+	using namespace UE::SceneState;
 
 	// Start all the Tasks that haven't started yet and that meet their pre-requisites
 	InContext.ForEachTask(*this,
-		[&InContext, &FindTaskInstance](const FSceneStateTask& InTask, FStructView InTaskInstance)->EIterationResult
+		[&InContext, StartingIndex=TaskRange.Index, &InStateInstance](const FSceneStateTask& InTask, uint16 InTaskRelativeIndex)->EIterationResult
 		{
-			FSceneStateTaskInstance* TaskInstance = InTaskInstance.GetPtr<FSceneStateTaskInstance>();
-			if (!TaskInstance || TaskInstance->GetStatus() != EExecutionStatus::NotStarted)
+			if (InStateInstance.GetTaskStatus(InTaskRelativeIndex) != EExecutionStatus::NotStarted)
 			{
 				return EIterationResult::Continue;
 			}
 
 			bool bPrerequisitesMet = true;
-			for (uint16 Prerequisite : InContext.GetTaskPrerequisites(InTask))
+			for (uint16 PrerequisiteIndex : InContext.GetTaskPrerequisites(InTask))
 			{
-				FSceneStateTaskInstance* const PrerequisiteTaskInstance = FindTaskInstance(Prerequisite);
-				if (!PrerequisiteTaskInstance || PrerequisiteTaskInstance->GetStatus() != EExecutionStatus::Finished)
+				// Prerequisite indices are stored relative
+				if (InStateInstance.GetTaskStatus(PrerequisiteIndex) != EExecutionStatus::Finished)
 				{
 					bPrerequisitesMet = false;
 					break;
@@ -191,7 +186,8 @@ void FSceneState::UpdateActiveTasks(const FSceneStateExecutionContext& InContext
 
 			if (bPrerequisitesMet)
 			{
-				InTask.Start(InContext, InTaskInstance);
+				const FStructView TaskInstance = InStateInstance.GetTaskInstances()[InTaskRelativeIndex];
+				InTask.Start(InContext, TaskInstance);
 			}
 
 			return EIterationResult::Continue;
@@ -213,8 +209,8 @@ bool FSceneState::HasPendingTasks(const FSceneStateExecutionContext& InContext)
 {
 	using namespace UE::SceneState;
 
-	FSceneStateInstance* const StateInstance = InContext.FindStateInstance(*this);
-	if (!StateInstance)
+	FSceneStateInstance* const Instance = InContext.FindStateInstance(*this);
+	if (!Instance || Instance->GetStatus() != EExecutionStatus::Running)
 	{
 		return false;
 	}
@@ -222,10 +218,9 @@ bool FSceneState::HasPendingTasks(const FSceneStateExecutionContext& InContext)
 	bool bHasPendingTask = false;
 
 	InContext.ForEachTask(*this,
-		[StateInstance, &bHasPendingTask](const FSceneStateTask& InTask, FStructView InTaskInstance)->EIterationResult
+		[&InContext, &bHasPendingTask, Instance](const FSceneStateTask& InTask, uint16 InTaskRelativeIndex)->EIterationResult
 		{
-			FSceneStateTaskInstance* TaskInstance = InTaskInstance.GetPtr<FSceneStateTaskInstance>();
-			if (TaskInstance && TaskInstance->GetStatus() != EExecutionStatus::Finished)
+			if (Instance->GetTaskStatus(InTaskRelativeIndex) != EExecutionStatus::Finished)
 			{
 				bHasPendingTask = true;
 				return EIterationResult::Break;
@@ -238,20 +233,19 @@ bool FSceneState::HasPendingTasks(const FSceneStateExecutionContext& InContext)
 
 void FSceneState::AllocateTaskInstances(const FSceneStateExecutionContext& InContext, TConstArrayView<FConstStructView> InTemplateTaskInstances) const
 {
-	if (FInstancedStructContainer* const TaskInstanceContainer = InContext.FindOrAddTaskInstanceContainer(*this))
-	{
-		// Copy the Template data
-		*TaskInstanceContainer = InTemplateTaskInstances;
-
-		// Instance each Template Object into the Instance data
-		InstanceTaskObjects(InContext.GetRootObject()
-			, UE::SceneState::GetStructViews(*TaskInstanceContainer)
-			, InTemplateTaskInstances
-			, [](FObjectDuplicationParameters& InParams)->UObject*
-			{
-				return StaticDuplicateObjectEx(InParams);
-			});
-	}
+	FSceneStateInstance& StateInstance = InContext.FindOrAddStateInstance(*this);
+
+	// Copy the Template data
+	StateInstance.GetTaskInstances() = InTemplateTaskInstances;
+
+	// Instance each Template Object into the Instance data
+	InstanceTaskObjects(InContext.GetRootObject()
+		, UE::SceneState::GetStructViews(StateInstance.GetTaskInstances())
+		, InTemplateTaskInstances
+		, [](FObjectDuplicationParameters& InParams)->UObject*
+		{
+			return StaticDuplicateObjectEx(InParams);
+		});
 }
 
 bool FSceneState::ApplyBindings(const FSceneStateExecutionContext& InContext, FSceneStateInstance& InStateInstance) const
@@ -261,7 +255,7 @@ bool FSceneState::ApplyBindings(const FSceneStateExecutionContext& InContext, FS
 	const UE::SceneState::FApplyBatchParams ApplyBatchParams
 		{
 			.BindingsBatch = BindingsBatch.Get(),
-			.TargetDataView = InStateInstance.GetParametersMutableView(),
+			.TargetDataView = InStateInstance.GetParameters().GetMutableValue(),
 		};
 
 	return ApplyBatch(InContext, ApplyBatchParams);
diff --git a/Engine/Plugins/VirtualProduction/SceneState/Source/SceneState/Private/SceneStateExecutionContext.cpp b/Engine/Plugins/VirtualProduction/SceneState/Source/SceneState/Private/SceneStateExecutionContext.cpp
--- a/Engine/Plugins/VirtualProduction/SceneState/Source/SceneState/Private/SceneStateExecutionContext.cpp
+++ b/Engine/Plugins/VirtualProduction/SceneState/Source/SceneState/Private/SceneStateExecutionContext.cpp
@@ -124,7 +124,6 @@ void FSceneStateExecutionContext::Reset()
 	StateMachineInstances.Reset();
 	TransitionInstances.Reset();
 	ConduitInstances.Reset();
-	TaskInstanceContainers.Reset();
 	StateMachineExecutionStack.Reset();
 	FunctionInstanceBatches.Reset();
 	FunctionInstanceViews.Reset();
@@ -208,7 +207,7 @@ FPropertyBindingDataView FSceneStateExecutionContext::FindDataView(const FSceneS
 	case ESceneStateDataType::State:
 		if (FSceneStateInstance* StateInstance = FindStateInstance(DataIndex))
 		{
-			return StateInstance->GetParametersMutableView();
+			return StateInstance->GetParameters().GetMutableValue();
 		}
 		break;
 	}
@@ -222,7 +221,7 @@ FPropertyBindingDataView FSceneStateExecutionContext::FindDataView(const FSceneS
 	return FPropertyBindingDataView();
 }
 
-void FSceneStateExecutionContext::ForEachTask(const FSceneState& InState, TFunctionRef<UE::SceneState::EIterationResult(const FSceneStateTask&, FStructView)> InCallable) const
+void FSceneStateExecutionContext::ForEachTask(const FSceneState& InState, TFunctionRef<UE::SceneState::EIterationResult(const FSceneStateTask&, uint16)> InCallable) const
 {
 	const FSceneStateRange TaskRange = InState.GetTaskRange();
 	if (TaskRange.Count == 0)
@@ -230,12 +229,6 @@ void FSceneStateExecutionContext::ForEachTask(const FSceneState& InState, TFunct
 		return;
 	}
 
-	FInstancedStructContainer* TaskInstanceContainer = FindTaskInstanceContainer(InState);
-	if (!TaskInstanceContainer)
-	{
-		return;
-	}
-
 	const FInstancedStructContainer& Tasks = TemplateData->Tasks;
 
 	checkfSlow(Tasks.IsValidIndex(TaskRange.Index) && Tasks.IsValidIndex(TaskRange.GetLastIndex())
@@ -245,18 +238,11 @@ void FSceneStateExecutionContext::ForEachTask(const FSceneState& InState, TFunct
 		, TaskRange.GetLastIndex()
 		, Tasks.Num());
 
-	for (uint16 TaskIndex = 0; TaskIndex < TaskRange.Count; ++TaskIndex)
+	for (uint16 RelativeIndex = 0; RelativeIndex < TaskRange.Count; ++RelativeIndex)
 	{
-		if (!TaskInstanceContainer->IsValidIndex(TaskIndex))
-		{
-			continue;
-		}
-
-		const uint16 AbsoluteIndex = TaskRange.Index + TaskIndex;
+		const uint16 AbsoluteIndex = TaskRange.Index + RelativeIndex;
 		const FSceneStateTask& Task = Tasks[AbsoluteIndex].Get<const FSceneStateTask>();
-		const FStructView TaskInstance = (*TaskInstanceContainer)[TaskIndex];
-
-		if (InCallable(Task, TaskInstance) == UE::SceneState::EIterationResult::Break)
+		if (InCallable(Task, RelativeIndex) == UE::SceneState::EIterationResult::Break)
 		{
 			break;
 		}
@@ -538,62 +524,70 @@ FConstStructView FSceneStateExecutionContext::FindTask(uint16 InAbsoluteIndex) c
 	return {};
 }
 
-FInstancedStructContainer* FSceneStateExecutionContext::FindOrAddTaskInstanceContainer(const FSceneState& InState) const
+uint16 FSceneStateExecutionContext::FindTaskRelativeIndex(const FSceneStateTask& InTask) const
 {
-	// Note: task instance containers are allocated per state
-	return &TaskInstanceContainers.FindOrAdd(GetStateIndex(InState));
-}
+	const uint16 StateIndex = InTask.GetParentStateIndex();
 
-FInstancedStructContainer* FSceneStateExecutionContext::FindTaskInstanceContainer(const FSceneState& InState) const
-{
-	return FindTaskInstanceContainer(GetStateIndex(InState));
-}
+	const TConstArrayView<FSceneState> StatesView = TemplateData->States;
 
-FInstancedStructContainer* FSceneStateExecutionContext::FindTaskInstanceContainer(uint16 InAbsoluteIndex) const
-{
-	return TaskInstanceContainers.Find(InAbsoluteIndex);
+	checkf(StatesView.IsValidIndex(StateIndex)
+		, TEXT("[Context: %s] Task Parent State Index %d is out of bounds! (Total State Count: %d)")
+		, *GetExecutionContextName()
+		, StateIndex
+		, StatesView.Num());
+
+	const FSceneState& State = StatesView[StateIndex];
+	return InTask.GetTaskIndex() - State.GetTaskRange().Index;
 }
 
-FStructView FSceneStateExecutionContext::FindTaskInstance(uint16 InAbsoluteIndex) const
+FStructView FSceneStateExecutionContext::FindTaskInstance(const FSceneStateTask& InTask) const
 {
-	const FInstancedStructContainer& Tasks = TemplateData->Tasks;
-	if (!Tasks.IsValidIndex(InAbsoluteIndex))
+	const uint16 StateIndex = InTask.GetParentStateIndex();
+
+	FSceneStateInstance* const StateInstance = FindStateInstance(StateIndex);
+	if (!StateInstance)
 	{
 		return {};
 	}
 
-	const FSceneStateTask& Task = Tasks[InAbsoluteIndex].Get<const FSceneStateTask>();
-
-	const uint16 StateIndex = Task.GetParentStateIndex();
+	FInstancedStructContainer& TaskInstanceContainer = StateInstance->GetTaskInstances();
 
-	FInstancedStructContainer* const TaskInstanceContainer = FindTaskInstanceContainer(StateIndex);
-	if (!TaskInstanceContainer)
+	const uint16 TaskRelativeIndex = FindTaskRelativeIndex(InTask);
+	if (!TaskInstanceContainer.IsValidIndex(TaskRelativeIndex))
 	{
 		return {};
 	}
 
-	const TConstArrayView<FSceneState> StatesView = TemplateData->States;
-
-	checkf(StatesView.IsValidIndex(StateIndex)
-		, TEXT("[Context: %s] Task Parent State Index %d is out of bounds! (Total State Count: %d)")
-		, *GetExecutionContextName()
-		, StateIndex
-		, StatesView.Num());
-
-	const FSceneState& State = StatesView[StateIndex];
+	return TaskInstanceContainer[TaskRelativeIndex];
+}
 
-	const uint16 TaskRelativeIndex = InAbsoluteIndex - State.GetTaskRange().Index;
-	if (!TaskInstanceContainer->IsValidIndex(TaskRelativeIndex))
+FStructView FSceneStateExecutionContext::FindTaskInstance(uint16 InTaskAbsoluteIndex) const
+{
+	const FInstancedStructContainer& Tasks = TemplateData->Tasks;
+	if (!Tasks.IsValidIndex(InTaskAbsoluteIndex))
 	{
 		return {};
 	}
 
-	return (*TaskInstanceContainer)[TaskRelativeIndex];
+	const FSceneStateTask& Task = Tasks[InTaskAbsoluteIndex].Get<const FSceneStateTask>();
+	return FindTaskInstance(Task);
+}
+
+UE::SceneState::EExecutionStatus FSceneStateExecutionContext::GetTaskStatus(const FSceneStateTask& InTask) const
+{
+	if (FSceneStateInstance* const StateInstance = FindStateInstance(InTask.GetParentStateIndex()))
+	{
+		return StateInstance->GetTaskStatus(FindTaskRelativeIndex(InTask));
+	}
+	return UE::SceneState::EExecutionStatus::NotStarted;
 }
 
-void FSceneStateExecutionContext::RemoveTaskInstanceContainer(const FSceneState& InState) const
+void FSceneStateExecutionContext::SetTaskStatus(const FSceneStateTask& InTask, UE::SceneState::EExecutionStatus InStatus) const
 {
-	TaskInstanceContainers.Remove(GetStateIndex(InState));
+	if (FSceneStateInstance* const StateInstance = FindStateInstance(InTask.GetParentStateIndex()))
+	{
+		StateInstance->SetTaskStatus(FindTaskRelativeIndex(InTask), InStatus);
+	}
 }
 
 FSceneStateMachineInstance* FSceneStateExecutionContext::FindOrAddStateMachineInstance(const FSceneStateMachine& InStateMachine) const
diff --git a/Engine/Plugins/VirtualProduction/SceneState/Source/SceneState/Private/SceneStateInstance.cpp b/Engine/Plugins/VirtualProduction/SceneState/Source/SceneState/Private/SceneStateInstance.cpp
--- a/Engine/Plugins/VirtualProduction/SceneState/Source/SceneState/Private/SceneStateInstance.cpp
+++ b/Engine/Plugins/VirtualProduction/SceneState/Source/SceneState/Private/SceneStateInstance.cpp
@@ -2,6 +2,8 @@
 
 #include "SceneStateInstance.h"
 
+#include "SceneState.h"
+
 namespace UE::SceneState::Private
 {
 	uint16 GetNextInstanceId()
@@ -26,9 +28,12 @@ FSceneStateInstance::FSceneStateInstance()
 }
 PRAGMA_ENABLE_DEPRECATION_WARNINGS
 
-void FSceneStateInstance::InitializeParameters(const FInstancedPropertyBag& InTemplateParameters)
+void FSceneStateInstance::Setup(const FSceneState& InState)
 {
-	Parameters = InTemplateParameters;
+	using namespace UE::SceneState;
+	SetStatus(EExecutionStatus::NotStarted);
+	Parameters = InState.GetParameters();
+	TaskStatuses.Init(EExecutionStatus::NotStarted, InState.GetTaskRange().Count);
 }
 
 UE::SceneState::EExecutionStatus FSceneStateInstance::GetStatus() const
@@ -49,3 +54,30 @@ uint16 FSceneStateInstance::GetInstanceId() const
 {
 	return InstanceId;
 }
+
+FInstancedPropertyBag& FSceneStateInstance::GetParameters()
+{
+	return Parameters;
+}
+
+FInstancedStructContainer& FSceneStateInstance::GetTaskInstances()
+{
+	return TaskInstances;
+}
+
+UE::SceneState::EExecutionStatus FSceneStateInstance::GetTaskStatus(uint16 InTaskRelativeIndex) const
+{
+	if (ensure(TaskStatuses.IsValidIndex(InTaskRelativeIndex)))
+	{
+		return TaskStatuses[InTaskRelativeIndex];
+	}
+	return UE::SceneState::EExecutionStatus::NotStarted;
+}
+
+void FSceneStateInstance::SetTaskStatus(uint16 InTaskRelativeIndex, UE::SceneState::EExecutionStatus InTaskStatus)
+{
+	if (ensure(TaskStatuses.IsValidIndex(InTaskRelativeIndex)))
+	{
+		TaskStatuses[InTaskRelativeIndex] = InTaskStatus;
+	}
+}
diff --git a/Engine/Plugins/VirtualProduction/SceneState/Source/SceneState/Private/SceneStateTemplateData.cpp b/Engine/Plugins/VirtualProduction/SceneState/Source/SceneState/Private/SceneStateTemplateData.cpp
--- a/Engine/Plugins/VirtualProduction/SceneState/Source/SceneState/Private/SceneStateTemplateData.cpp
+++ b/Engine/Plugins/VirtualProduction/SceneState/Source/SceneState/Private/SceneStateTemplateData.cpp
@@ -171,54 +171,27 @@ const FSceneStateTask* USceneStateTemplateData::FindTaskFromNode(FObjectKey InTa
 	return nullptr;
 }
 
-void USceneStateTemplateData::ForEachStateInstance(const TSharedRef<UE::SceneState::FExecutionContextRegistry>& InContextRegistry, uint16 InStateIndex, TFunctionRef<void(const FSceneStateInstance&)> InFunctor) const
-{
-	if (!States.IsValidIndex(InStateIndex))
-	{
-		return;
-	}
-
-	InContextRegistry->ForEachExecutionContext(
-		[InStateIndex, &InFunctor](const FSceneStateExecutionContext& InExecutionContext)
-		{
-			if (const FSceneStateInstance* StateInstance = InExecutionContext.FindStateInstance(InStateIndex))
-			{
-				InFunctor(*StateInstance);
-			}
-		});
-}
-
 void USceneStateTemplateData::ForEachStateInstance(const TSharedRef<UE::SceneState::FExecutionContextRegistry>& InContextRegistry, FObjectKey InStateNode, TFunctionRef<void(const FSceneStateInstance&)> InFunctor) const
 {
-	if (const uint16* StateIndex = StateNodeToIndex.Find(InStateNode))
-	{
-		ForEachStateInstance(InContextRegistry, *StateIndex, InFunctor);
-	}
 }
 
-void USceneStateTemplateData::ForEachTaskInstance(const TSharedRef<UE::SceneState::FExecutionContextRegistry>& InContextRegistry, uint16 InTaskIndex, TFunctionRef<void(const FSceneStateTaskInstance&)> InFunctor) const
+void USceneStateTemplateData::ForEachTaskStatus(const TSharedRef<UE::SceneState::FExecutionContextRegistry>& InContextRegistry, uint16 InTaskIndex, TFunctionRef<void(UE::SceneState::EExecutionStatus)> InFunctor) const
 {
 	if (!Tasks.IsValidIndex(InTaskIndex))
 	{
 		return;
 	}
 
+	const FSceneStateTask& Task = Tasks[InTaskIndex].Get<const FSceneStateTask>();
 	InContextRegistry->ForEachExecutionContext(
-		[InTaskIndex, &InFunctor](const FSceneStateExecutionContext& InExecutionContext)
+		[&Task, &InFunctor](const FSceneStateExecutionContext& InExecutionContext)
 		{
-			if (const FSceneStateTaskInstance* TaskInstance = InExecutionContext.FindTaskInstance(InTaskIndex).GetPtr<const FSceneStateTaskInstance>())
-			{
-				InFunctor(*TaskInstance);
-			}
+			InFunctor(InExecutionContext.GetTaskStatus(Task));
 		});
 }
 
 void USceneStateTemplateData::ForEachTaskInstance(const TSharedRef<UE::SceneState::FExecutionContextRegistry>& InContextRegistry, FObjectKey InTaskNode, TFunctionRef<void(const FSceneStateTaskInstance&)> InFunctor) const
 {
-	if (const uint16* TaskIndex = TaskNodeToIndex.Find(InTaskNode))
-	{
-		ForEachTaskInstance(InContextRegistry, *TaskIndex, InFunctor);
-	}
 }
 
 void USceneStateTemplateData::OnObjectsReinstanced(const UClass* InOwnerClass, const TMap<UObject*, UObject*>& InReplacementMap)
diff --git a/Engine/Plugins/VirtualProduction/SceneState/Source/SceneState/Private/Tasks/SceneStateTask.cpp b/Engine/Plugins/VirtualProduction/SceneState/Source/SceneState/Private/Tasks/SceneStateTask.cpp
--- a/Engine/Plugins/VirtualProduction/SceneState/Source/SceneState/Private/Tasks/SceneStateTask.cpp
+++ b/Engine/Plugins/VirtualProduction/SceneState/Source/SceneState/Private/Tasks/SceneStateTask.cpp
@@ -66,7 +66,8 @@ const FSceneStateTaskBindingExtension* FSceneStateTask::GetBindingExtension() co
 
 FStructView FSceneStateTask::FindTaskInstance(const FSceneStateExecutionContext& InContext) const
 {
-	return InContext.FindTaskInstance(GetTaskIndex());
+	ensureAlwaysMsgf(false, TEXT("FSceneStateTask::FindTaskInstance is deprecated! Use the task instance provided by the virtual methods."));
+	return InContext.FindTaskInstance(*this);
 }
 
 void FSceneStateTask::Setup(const FSceneStateExecutionContext& InContext, FStructView InTaskInstance) const
@@ -75,13 +76,12 @@ void FSceneStateTask::Setup(const FSceneStateExecutionContext& InContext, FStruc
 
 	QUICK_SCOPE_CYCLE_COUNTER(SceneState_OnSetup);
 
-	FSceneStateTaskInstance* Instance = InTaskInstance.GetPtr<FSceneStateTaskInstance>();
-	if (!Instance)
+	if (!InTaskInstance.IsValid())
 	{
 		return;
 	}
 
-	Instance->SetStatus(EExecutionStatus::NotStarted);
+	InContext.SetTaskStatus(*this, EExecutionStatus::NotStarted);
 	InContext.SetupFunctionInstances(BindingsBatch);
 
 	if (const FSceneStateTaskBindingExtension* BindingExtension = GetBindingExtension())
@@ -102,16 +102,13 @@ void FSceneStateTask::Start(const FSceneStateExecutionContext& InContext, FStruc
 
 	QUICK_SCOPE_CYCLE_COUNTER(SceneState_OnStart);
 
-	FSceneStateTaskInstance* Instance = InTaskInstance.GetPtr<FSceneStateTaskInstance>();
-	if (!Instance || Instance->GetStatus() != EExecutionStatus::NotStarted)
+	if (!InTaskInstance.IsValid() || InContext.GetTaskStatus(*this) != EExecutionStatus::NotStarted)
 	{
 		return;
 	}
-
-	Instance->SetStatus(EExecutionStatus::Running);
+	InContext.SetTaskStatus(*this, EExecutionStatus::Running);
 
 	ApplyBindings(InContext, InTaskInstance);
-
 	OnStart(InContext, InTaskInstance);
 }
 
@@ -127,8 +124,7 @@ void FSceneStateTask::Tick(const FSceneStateExecutionContext& InContext, FStruct
 
 	QUICK_SCOPE_CYCLE_COUNTER(SceneState_OnTick);
 
-	FSceneStateTaskInstance* Instance = InTaskInstance.GetPtr<FSceneStateTaskInstance>();
-	if (!Instance || Instance->GetStatus() != EExecutionStatus::Running)
+	if (!InTaskInstance.IsValid() || InContext.GetTaskStatus(*this) != EExecutionStatus::Running)
 	{
 		return;
 	}
@@ -142,13 +138,12 @@ void FSceneStateTask::Stop(const FSceneStateExecutionContext& InContext, FStruct
 
 	QUICK_SCOPE_CYCLE_COUNTER(SceneState_OnStop);
 
-	FSceneStateTaskInstance* Instance = InTaskInstance.GetPtr<FSceneStateTaskInstance>();
-	if (!Instance)
+	if (!InTaskInstance.IsValid())
 	{
 		return;
 	}
 
-	if (Instance->GetStatus() == EExecutionStatus::Running)
+	if (InContext.GetTaskStatus(*this) == EExecutionStatus::Running)
 	{
 		OnStop(InContext, InTaskInstance, InStopReason);
 	}
@@ -164,7 +159,7 @@ void FSceneStateTask::Stop(const FSceneStateExecutionContext& InContext, FStruct
 			});
 	}
 
-	Instance->SetStatus(Private::GetTaskStatus(InStopReason));
+	InContext.SetTaskStatus(*this, Private::GetTaskStatus(InStopReason));
 }
 
 void FSceneStateTask::Finish(const FSceneStateExecutionContext& InContext, FStructView InTaskInstance) const
diff --git a/Engine/Plugins/VirtualProduction/SceneState/Source/SceneState/Private/Tasks/SceneStateTaskInstance.cpp b/Engine/Plugins/VirtualProduction/SceneState/Source/SceneState/Private/Tasks/SceneStateTaskInstance.cpp
deleted file mode 100644
diff --git a/Engine/Plugins/VirtualProduction/SceneState/Source/SceneState/Public/SceneState.h b/Engine/Plugins/VirtualProduction/SceneState/Source/SceneState/Public/SceneState.h
--- a/Engine/Plugins/VirtualProduction/SceneState/Source/SceneState/Public/SceneState.h
+++ b/Engine/Plugins/VirtualProduction/SceneState/Source/SceneState/Public/SceneState.h
@@ -12,6 +12,7 @@
 #define UE_API SCENESTATE_API
 
 struct FConstStructView;
+struct FInstancedStructContainer;
 struct FSceneStateExecutionContext;
 struct FSceneStateInstance;
 struct FStructView;
@@ -55,10 +56,11 @@ struct FSceneState
 	/** Called when the state is no longer active in its parent state machine */
 	UE_API void Exit(const FSceneStateExecutionContext& InContext) const;
 
-	const UStruct* GetParametersStruct() const
-	{
-		return Parameters.GetPropertyBagStruct();
-	}
+	/** Get the template parameters of the state */
+	const FInstancedPropertyBag& GetParameters() const;
+
+	/** Get the template parameter struct */
+	const UStruct* GetParametersStruct() const;
 
 	FSceneStateRange GetTaskRange() const
 	{
@@ -81,7 +83,7 @@ struct FSceneState
 	}
 
 	/** Starts pending tasks whose prerequisites are satisfied */
-	void UpdateActiveTasks(const FSceneStateExecutionContext& InContext) const;
+	void UpdateActiveTasks(const FSceneStateExecutionContext& InContext, FSceneStateInstance& InStateInstance) const;
 
 	/** Returns the State Name used for logging. Returns an empty string view if metadata is not available */
 	FStringView GetStateName(const FSceneStateExecutionContext& InContext) const;
@@ -90,14 +92,14 @@ struct FSceneState
 	UE_API bool HasPendingTasks(const FSceneStateExecutionContext& InContext) const;
 
 	/** Called on start to create the required task instances that will run for this state */
-	UE_API void AllocateTaskInstances(const FSceneStateExecutionContext& InContext, TConstArrayView<FConstStructView> InTemplateTaskInstances) const;
+	UE_INTERNAL UE_API void AllocateTaskInstances(const FSceneStateExecutionContext& InContext, TConstArrayView<FConstStructView> InTemplateTaskInstances) const;
 
 private:
 	/** Applies Bindings to the given state instance */
 	bool ApplyBindings(const FSceneStateExecutionContext& InContext, FSceneStateInstance& InStateInstance) const;
 
 	/** Instances the instanced references in source data into the target data using the provided duplication functor */
-	UE_API void InstanceTaskObjects(UObject* InOuter
+	UE_INTERNAL UE_API void InstanceTaskObjects(UObject* InOuter
 		, TConstArrayView<FStructView> InTargets
 		, TConstArrayView<FConstStructView> InSources
 		, TFunctionRef<UObject*(FObjectDuplicationParameters&)> InDuplicationFunc) const;
diff --git a/Engine/Plugins/VirtualProduction/SceneState/Source/SceneState/Public/SceneStateExecutionContext.h b/Engine/Plugins/VirtualProduction/SceneState/Source/SceneState/Public/SceneStateExecutionContext.h
--- a/Engine/Plugins/VirtualProduction/SceneState/Source/SceneState/Public/SceneStateExecutionContext.h
+++ b/Engine/Plugins/VirtualProduction/SceneState/Source/SceneState/Public/SceneStateExecutionContext.h
@@ -5,6 +5,7 @@
 #include "Conduit/SceneStateConduitInstance.h"
 #include "Containers/Array.h"
 #include "Functions/SceneStateFunctionInstanceBatch.h"
+#include "SceneStateEnums.h"
 #include "SceneStateExecutionContextHandle.h"
 #include "SceneStateInstance.h"
 #include "SceneStateMachineInstance.h"
@@ -141,7 +142,7 @@ struct FSceneStateExecutionContext
 	}
 
 	/** Invokes a given callable for each task in the given state */
-	void ForEachTask(const FSceneState& InState, TFunctionRef<UE::SceneState::EIterationResult(const FSceneStateTask&, FStructView)> InCallable) const;
+	void ForEachTask(const FSceneState& InState, TFunctionRef<UE::SceneState::EIterationResult(const FSceneStateTask&, uint16 /*RelativeIndex*/)> InCallable) const;
 
 	/** Returns an array of const views of the Template Task Instances of the given State */
 	TArray<FConstStructView> GetTemplateTaskInstances(const FSceneState& InState) const;
@@ -211,20 +212,20 @@ struct FSceneStateExecutionContext
 	/** Returns the Task of this context for the given task absolute index. Returns an invalid view if not found */
 	FConstStructView FindTask(uint16 InAbsoluteIndex) const;
 
-	/** Returns the Task Instance container of this context for the given state, adding a new one if not found */
-	FInstancedStructContainer* FindOrAddTaskInstanceContainer(const FSceneState& InState) const;
+	/** Gets the task index relative to the other tasks in the state */
+	uint16 FindTaskRelativeIndex(const FSceneStateTask& InTask) const;
 
-	/** Returns the Task Instance container of this context for the given state, returning null if not found */
-	FInstancedStructContainer* FindTaskInstanceContainer(const FSceneState& InState) const;
-
-	/** Returns the Task Instance container of this context for the given state absolute index, returning null if not found */
-	FInstancedStructContainer* FindTaskInstanceContainer(uint16 InAbsoluteIndex) const;
+	/** Returns the Task Instance of this context for the given task. Returns an invalid view if not found */
+	UE_INTERNAL UE_API FStructView FindTaskInstance(const FSceneStateTask& InTask) const;
 
 	/** Returns the Task Instance of this context for the given task absolute index. Returns an invalid view if not found */
-	FStructView FindTaskInstance(uint16 InAbsoluteIndex) const;
+	FStructView FindTaskInstance(uint16 InTaskAbsoluteIndex) const;
+
+	/** Returns the given task's execution status */
+	UE_INTERNAL UE_API UE::SceneState::EExecutionStatus GetTaskStatus(const FSceneStateTask& InTask) const;
 
-	/** Removes the Task instance container of this context for the given state*/
-	void RemoveTaskInstanceContainer(const FSceneState& InState) const;
+	/** Set the task status of this context for the given task */
+	void SetTaskStatus(const FSceneStateTask& InTask, UE::SceneState::EExecutionStatus InStatus) const;
 
 	/** Returns the State Machine Instance of this context for the given state machine, adding a new one if it's not found */
 	FSceneStateMachineInstance* FindOrAddStateMachineInstance(const FSceneStateMachine& InStateMachine) const;
@@ -317,13 +318,6 @@ private:
 	UPROPERTY()
 	mutable TMap<uint16, FSceneStateInstance> StateInstances;
 
-	/**
-	 * Map of State Index to the Task Instance Container
-	 * It's allocated when the state starts and removed on exit.
-	 */
-	UPROPERTY()
-	mutable TMap<uint16, FInstancedStructContainer> TaskInstanceContainers;
-
 	/**
  	 * Map of State Machine Index to its Instance Data
  	 * It's allocated when the State Machine starts and removed on exit.
diff --git a/Engine/Plugins/VirtualProduction/SceneState/Source/SceneState/Public/SceneStateInstance.h b/Engine/Plugins/VirtualProduction/SceneState/Source/SceneState/Public/SceneStateInstance.h
--- a/Engine/Plugins/VirtualProduction/SceneState/Source/SceneState/Public/SceneStateInstance.h
+++ b/Engine/Plugins/VirtualProduction/SceneState/Source/SceneState/Public/SceneStateInstance.h
@@ -3,11 +3,14 @@
 #pragma once
 
 #include "SceneStateEnums.h"
+#include "StructUtils/InstancedStructContainer.h"
 #include "StructUtils/PropertyBag.h"
 #include "SceneStateInstance.generated.h"
 
 #define UE_API SCENESTATE_API
 
+struct FSceneState;
+
 /** Instance data of a State */
 USTRUCT()
 struct FSceneStateInstance
@@ -22,12 +25,12 @@ PRAGMA_DISABLE_DEPRECATION_WARNINGS
 	FSceneStateInstance& operator=(FSceneStateInstance&&) = default;
 PRAGMA_ENABLE_DEPRECATION_WARNINGS
 
+	/** Called to setup the state instance for use */
+	void Setup(const FSceneState& InState);
+
 	UE_DEPRECATED(5.8, "ElapsedTime is deprecated, it is no longer in use.")
 	float ElapsedTime = 0.f;
 
-	/** Initializes this instance's parameters to the template's */
-	void InitializeParameters(const FInstancedPropertyBag& InTemplateParameters);
-
 	/** Gets the execution status of the state instance */
 	UE_API UE::SceneState::EExecutionStatus GetStatus() const;
 
@@ -41,16 +44,30 @@ PRAGMA_ENABLE_DEPRECATION_WARNINGS
 	/** Get an identifier for this instance */
 	UE_INTERNAL UE_API uint16 GetInstanceId() const;
 
-	FStructView GetParametersMutableView()
-	{
-		return Parameters.GetMutableValue();
-	}
+	/** Returns the state parameters */
+	FInstancedPropertyBag& GetParameters();
+
+	/** Returns the instance data of the tasks owned by this state */
+	FInstancedStructContainer& GetTaskInstances();
+
+	/** Returns the status of the task at the given index */
+	UE::SceneState::EExecutionStatus GetTaskStatus(uint16 InTaskRelativeIndex) const;
+
+	/** Sets the status of the task at the given index */
+	void SetTaskStatus(uint16 InTaskRelativeIndex, UE::SceneState::EExecutionStatus InTaskStatus);
 
 private:
 	/** Instanced state parameters */
 	UPROPERTY()
 	FInstancedPropertyBag Parameters;
 
+	/** The task instances owned by this state */
+	UPROPERTY()
+	FInstancedStructContainer TaskInstances;
+
+	/** Execution status of the tasks */
+	TArray<UE::SceneState::EExecutionStatus> TaskStatuses;
+
 	/** The id for this instance. This is used to differentiate state instances for a same state and consequently task instances for a same task*/
 	uint16 InstanceId;
 };
diff --git a/Engine/Plugins/VirtualProduction/SceneState/Source/SceneState/Public/SceneStateTemplateData.h b/Engine/Plugins/VirtualProduction/SceneState/Source/SceneState/Public/SceneStateTemplateData.h
--- a/Engine/Plugins/VirtualProduction/SceneState/Source/SceneState/Public/SceneStateTemplateData.h
+++ b/Engine/Plugins/VirtualProduction/SceneState/Source/SceneState/Public/SceneStateTemplateData.h
@@ -72,32 +72,22 @@ public:
 	/** Finds the mapped compiled task for a given task node */
 	UE_API const FSceneStateTask* FindTaskFromNode(FObjectKey InTaskNode) const;
 
-	/**
-	 * For a given root state and state index retrieves the mapped state instances
-	 * @param InContextRegistry the registry containing all the execution contexts
-	 * @param InStateIndex the state index to look for
-	 * @param InFunctor the functor to run for each state instance
-	 */
-	UE_API void ForEachStateInstance(const TSharedRef<UE::SceneState::FExecutionContextRegistry>& InContextRegistry
-		, uint16 InStateIndex
-		, TFunctionRef<void(const FSceneStateInstance&)> InFunctor) const;
-
-	UE_DEPRECATED(5.8, "ForEachStateInstance deprecated. No longer supported. Use the overload that takes in a state index instead")
+	UE_DEPRECATED(5.8, "ForEachStateInstance deprecated. No longer supported.")
 	UE_API void ForEachStateInstance(const TSharedRef<UE::SceneState::FExecutionContextRegistry>& InContextRegistry
 		, FObjectKey InStateNode
 		, TFunctionRef<void(const FSceneStateInstance&)> InFunctor) const;
 
 	/**
-	 * For a given root state and task index retrieves the mapped task instances
+	 * For a given root state and task index retrieves the mapped task status
 	 * @param InContextRegistry the registry containing all the execution contexts
 	 * @param InTaskIndex the task index to look for
-	 * @param InFunctor the functor to run for each task instance
+	 * @param InFunctor the functor to run for each task status
 	 */
-	UE_API void ForEachTaskInstance(const TSharedRef<UE::SceneState::FExecutionContextRegistry>& InContextRegistry
+	UE_INTERNAL UE_API void ForEachTaskStatus(const TSharedRef<UE::SceneState::FExecutionContextRegistry>& InContextRegistry
 		, uint16 InTaskIndex
-		, TFunctionRef<void(const FSceneStateTaskInstance&)> InFunctor) const;
+		, TFunctionRef<void(UE::SceneState::EExecutionStatus)> InFunctor) const;
 
-	UE_DEPRECATED(5.8, "ForEachTaskInstance deprecated. No longer supported. Use the overload that takes in a task index instead")
+	UE_DEPRECATED(5.8, "ForEachTaskInstance deprecated. No longer supported.")
 	UE_API void ForEachTaskInstance(const TSharedRef<UE::SceneState::FExecutionContextRegistry>& InContextRegistry
 		, FObjectKey InTaskNode
 		, TFunctionRef<void(const FSceneStateTaskInstance&)> InFunctor) const;
diff --git a/Engine/Plugins/VirtualProduction/SceneState/Source/SceneState/Public/Tasks/SceneStateTask.h b/Engine/Plugins/VirtualProduction/SceneState/Source/SceneState/Public/Tasks/SceneStateTask.h
--- a/Engine/Plugins/VirtualProduction/SceneState/Source/SceneState/Public/Tasks/SceneStateTask.h
+++ b/Engine/Plugins/VirtualProduction/SceneState/Source/SceneState/Public/Tasks/SceneStateTask.h
@@ -73,7 +73,7 @@ struct FSceneStateTask
 	/** Called to get the optional Binding Extension for a Task */
 	UE_API const FSceneStateTaskBindingExtension* GetBindingExtension() const;
 
-	/** Finds the Task Instance for the given Execution Instance */
+	UE_DEPRECATED(5.8, "FindTaskInstance is deprecated. Use the task instance provided by the task's virtual method ")
 	UE_API FStructView FindTaskInstance(const FSceneStateExecutionContext& InContext) const;
 
 	/**
diff --git a/Engine/Plugins/VirtualProduction/SceneState/Source/SceneState/Public/Tasks/SceneStateTaskInstance.h b/Engine/Plugins/VirtualProduction/SceneState/Source/SceneState/Public/Tasks/SceneStateTaskInstance.h
--- a/Engine/Plugins/VirtualProduction/SceneState/Source/SceneState/Public/Tasks/SceneStateTaskInstance.h
+++ b/Engine/Plugins/VirtualProduction/SceneState/Source/SceneState/Public/Tasks/SceneStateTaskInstance.h
@@ -5,24 +5,15 @@
 #include "SceneStateEnums.h"
 #include "SceneStateTaskInstance.generated.h"
 
-struct FPropertyBindingDataView;
-
 /** Represents the Instance Data of a Task */
 USTRUCT()
 struct FSceneStateTaskInstance
 {
 	GENERATED_BODY()
 
-	/** Sets the execution status of the Task */
-	void SetStatus(UE::SceneState::EExecutionStatus InStatus);
-
-	/** Gets the execution status of the task */
+	UE_DEPRECATED(5.8, "Status has been moved outside of task instance and now managed by the execution context ")
 	UE::SceneState::EExecutionStatus GetStatus() const
 	{
-		return Status;
+		return UE::SceneState::EExecutionStatus::NotStarted;
 	}
-
-private:
-	/** Current status of the Task */
-	UE::SceneState::EExecutionStatus Status = UE::SceneState::EExecutionStatus::NotStarted;
 };
diff --git a/Engine/Plugins/VirtualProduction/SceneState/Source/SceneStateBlueprintEditor/Private/DebugExecutors/SceneStateDebugTaskExecutor.cpp b/Engine/Plugins/VirtualProduction/SceneState/Source/SceneStateBlueprintEditor/Private/DebugExecutors/SceneStateDebugTaskExecutor.cpp
--- a/Engine/Plugins/VirtualProduction/SceneState/Source/SceneStateBlueprintEditor/Private/DebugExecutors/SceneStateDebugTaskExecutor.cpp
+++ b/Engine/Plugins/VirtualProduction/SceneState/Source/SceneStateBlueprintEditor/Private/DebugExecutors/SceneStateDebugTaskExecutor.cpp
@@ -65,11 +65,10 @@ void FDebugTaskExecutor::OnStart(const FSceneStateExecutionContext& InExecutionC
 	{
 		Setup(InExecutionContext, *Task);
 
-		const FStructView TaskInstance = Task->FindTaskInstance(InExecutionContext);
-
+		const FStructView TaskInstance = InExecutionContext.FindTaskInstance(*Task);
 		Task->Setup(InExecutionContext, TaskInstance);
 		Task->Start(InExecutionContext, TaskInstance);
-		ConditionallyExit(TaskInstance);
+		ConditionallyExit(*Task, InExecutionContext);
 	}
 }
 
@@ -77,10 +76,9 @@ void FDebugTaskExecutor::OnTick(const FSceneStateExecutionContext& InExecutionCo
 {
 	if (const FSceneStateTask* Task = GetTask(InExecutionContext))
 	{
-		const FStructView TaskInstance = Task->FindTaskInstance(InExecutionContext);
-
+		const FStructView TaskInstance = InExecutionContext.FindTaskInstance(*Task);
 		Task->Tick(InExecutionContext, TaskInstance, InDeltaSeconds);
-		ConditionallyExit(TaskInstance);
+		ConditionallyExit(*Task, InExecutionContext);
 	}
 }
 
@@ -88,15 +86,14 @@ void FDebugTaskExecutor::OnExit(const FSceneStateExecutionContext& InExecutionCo
 {
 	if (const FSceneStateTask* Task = GetTask(InExecutionContext))
 	{
-		const FStructView TaskInstance = Task->FindTaskInstance(InExecutionContext);
+		const FStructView TaskInstance = InExecutionContext.FindTaskInstance(*Task);
 		Task->Stop(InExecutionContext, TaskInstance, ESceneStateTaskStopReason::State);
 	}
 }
 
-void FDebugTaskExecutor::ConditionallyExit(FStructView InTaskInstance)
+void FDebugTaskExecutor::ConditionallyExit(const FSceneStateTask& InTask, const FSceneStateExecutionContext& InExecutionContext)
 {
-	const FSceneStateTaskInstance* TaskInstance = InTaskInstance.GetPtr<FSceneStateTaskInstance>();
-	if (TaskInstance && TaskInstance->GetStatus() == EExecutionStatus::Finished)
+	if (InExecutionContext.GetTaskStatus(InTask) == EExecutionStatus::Finished)
 	{
 		Exit();
 	}
diff --git a/Engine/Plugins/VirtualProduction/SceneState/Source/SceneStateBlueprintEditor/Private/DebugExecutors/SceneStateDebugTaskExecutor.h b/Engine/Plugins/VirtualProduction/SceneState/Source/SceneStateBlueprintEditor/Private/DebugExecutors/SceneStateDebugTaskExecutor.h
--- a/Engine/Plugins/VirtualProduction/SceneState/Source/SceneStateBlueprintEditor/Private/DebugExecutors/SceneStateDebugTaskExecutor.h
+++ b/Engine/Plugins/VirtualProduction/SceneState/Source/SceneStateBlueprintEditor/Private/DebugExecutors/SceneStateDebugTaskExecutor.h
@@ -34,7 +34,7 @@ protected:
 
 private:
 	/** Checks whether the task completed and calls exit if so */
-	void ConditionallyExit(FStructView InTaskInstance);
+	void ConditionallyExit(const FSceneStateTask& InTask, const FSceneStateExecutionContext& InExecutionContext);
 };
 
 } // UE::SceneState::Editor
diff --git a/Engine/Plugins/VirtualProduction/SceneState/Source/SceneStateMachineEditor/Private/Widgets/SSceneStateMachineTaskNode.cpp b/Engine/Plugins/VirtualProduction/SceneState/Source/SceneStateMachineEditor/Private/Widgets/SSceneStateMachineTaskNode.cpp
--- a/Engine/Plugins/VirtualProduction/SceneState/Source/SceneStateMachineEditor/Private/Widgets/SSceneStateMachineTaskNode.cpp
+++ b/Engine/Plugins/VirtualProduction/SceneState/Source/SceneStateMachineEditor/Private/Widgets/SSceneStateMachineTaskNode.cpp
@@ -204,17 +204,17 @@ void SStateMachineTaskNode::GetNodeInfoPopups(FNodeInfoContext* InContext, TArra
 
 	TOptional<EExecutionStatus> ResultStatus;
 
-	TemplateData->ForEachTaskInstance(RootState->GetContextRegistry(), *Node->GetCompiledIndex(),
-		[&OutPopups, &ResultStatus, &StatusColors = StatusColors](const FSceneStateTaskInstance& InTaskInstance)
+	TemplateData->ForEachTaskStatus(RootState->GetContextRegistry(), *Node->GetCompiledIndex(),
+		[&OutPopups, &ResultStatus, &StatusColors = StatusColors](UE::SceneState::EExecutionStatus InTaskStatus)
 		{
 			// Get the lowest status to show
-			if (InTaskInstance.GetStatus() <= ResultStatus.Get(EExecutionStatus::Finished))
+			if (InTaskStatus <= ResultStatus.Get(EExecutionStatus::Finished))
 			{
-				ResultStatus = InTaskInstance.GetStatus();
+				ResultStatus = InTaskStatus;
 			}
 
 			FText Message;
-			switch (InTaskInstance.GetStatus())
+			switch (InTaskStatus)
 			{
 			case EExecutionStatus::Running:
 				Message = LOCTEXT("RunningStatusFormat", "Running");
@@ -228,7 +228,7 @@ void SStateMachineTaskNode::GetNodeInfoPopups(FNodeInfoContext* InContext, TArra
 				return;
 			}
 
-			OutPopups.Emplace(nullptr, StatusColors[InTaskInstance.GetStatus()], Message.ToString());
+			OutPopups.Emplace(nullptr, StatusColors[InTaskStatus], Message.ToString());
 		});
 
 	// If the result status was never gotten (i.e. no Task Instances), set to Not-Started
