Motion Design Scene State: changed GetTaskStatus to return optional instead of a fixed value in case the task is not found in the context

#Synced-CL 50950138

diff --git a/Engine/Plugins/VirtualProduction/SceneState/Source/SceneState/Private/SceneStateExecutionContext.cpp b/Engine/Plugins/VirtualProduction/SceneState/Source/SceneState/Private/SceneStateExecutionContext.cpp
--- a/Engine/Plugins/VirtualProduction/SceneState/Source/SceneState/Private/SceneStateExecutionContext.cpp
+++ b/Engine/Plugins/VirtualProduction/SceneState/Source/SceneState/Private/SceneStateExecutionContext.cpp
@@ -573,13 +573,13 @@ FStructView FSceneStateExecutionContext::FindTaskInstance(uint16 InTaskAbsoluteI
 	return FindTaskInstance(Task);
 }
 
-UE::SceneState::EExecutionStatus FSceneStateExecutionContext::GetTaskStatus(const FSceneStateTask& InTask) const
+TOptional<UE::SceneState::EExecutionStatus> FSceneStateExecutionContext::GetTaskStatus(const FSceneStateTask& InTask) const
 {
 	if (FSceneStateInstance* const StateInstance = FindStateInstance(InTask.GetParentStateIndex()))
 	{
 		return StateInstance->GetTaskStatus(FindTaskRelativeIndex(InTask));
 	}
-	return UE::SceneState::EExecutionStatus::NotStarted;
+	return {};
 }
 
 void FSceneStateExecutionContext::SetTaskStatus(const FSceneStateTask& InTask, UE::SceneState::EExecutionStatus InStatus) const
diff --git a/Engine/Plugins/VirtualProduction/SceneState/Source/SceneState/Private/SceneStateTemplateData.cpp b/Engine/Plugins/VirtualProduction/SceneState/Source/SceneState/Private/SceneStateTemplateData.cpp
--- a/Engine/Plugins/VirtualProduction/SceneState/Source/SceneState/Private/SceneStateTemplateData.cpp
+++ b/Engine/Plugins/VirtualProduction/SceneState/Source/SceneState/Private/SceneStateTemplateData.cpp
@@ -186,7 +186,11 @@ void USceneStateTemplateData::ForEachTaskStatus(const TSharedRef<UE::SceneState:
 	InContextRegistry->ForEachExecutionContext(
 		[&Task, &InFunctor](const FSceneStateExecutionContext& InExecutionContext)
 		{
-			InFunctor(InExecutionContext.GetTaskStatus(Task));
+			const TOptional<UE::SceneState::EExecutionStatus> TaskStatus = InExecutionContext.GetTaskStatus(Task);
+			if (TaskStatus.IsSet())
+			{
+				InFunctor(*TaskStatus);
+			}
 		});
 }
 
diff --git a/Engine/Plugins/VirtualProduction/SceneState/Source/SceneState/Private/Tasks/SceneStateTask.cpp b/Engine/Plugins/VirtualProduction/SceneState/Source/SceneState/Private/Tasks/SceneStateTask.cpp
--- a/Engine/Plugins/VirtualProduction/SceneState/Source/SceneState/Private/Tasks/SceneStateTask.cpp
+++ b/Engine/Plugins/VirtualProduction/SceneState/Source/SceneState/Private/Tasks/SceneStateTask.cpp
@@ -102,10 +102,11 @@ void FSceneStateTask::Start(const FSceneStateExecutionContext& InContext, FStruc
 
 	QUICK_SCOPE_CYCLE_COUNTER(SceneState_OnStart);
 
-	if (!InTaskInstance.IsValid() || InContext.GetTaskStatus(*this) != EExecutionStatus::NotStarted)
+	if (!InTaskInstance.IsValid() || !IsTaskStatus(InContext, EExecutionStatus::NotStarted))
 	{
 		return;
 	}
+
 	InContext.SetTaskStatus(*this, EExecutionStatus::Running);
 
 	ApplyBindings(InContext, InTaskInstance);
@@ -124,7 +125,7 @@ void FSceneStateTask::Tick(const FSceneStateExecutionContext& InContext, FStruct
 
 	QUICK_SCOPE_CYCLE_COUNTER(SceneState_OnTick);
 
-	if (!InTaskInstance.IsValid() || InContext.GetTaskStatus(*this) != EExecutionStatus::Running)
+	if (!InTaskInstance.IsValid() || !IsTaskStatus(InContext, EExecutionStatus::Running))
 	{
 		return;
 	}
@@ -143,7 +144,7 @@ void FSceneStateTask::Stop(const FSceneStateExecutionContext& InContext, FStruct
 		return;
 	}
 
-	if (InContext.GetTaskStatus(*this) == EExecutionStatus::Running)
+	if (IsTaskStatus(InContext, EExecutionStatus::Running))
 	{
 		OnStop(InContext, InTaskInstance, InStopReason);
 	}
@@ -177,6 +178,12 @@ void FSceneStateTask::ClearFlags(ESceneStateTaskFlags InFlags)
 	TaskFlags &= ~InFlags;
 }
 
+bool FSceneStateTask::IsTaskStatus(const FSceneStateExecutionContext& InContext, UE::SceneState::EExecutionStatus InExpectedStatus) const
+{
+	const TOptional<UE::SceneState::EExecutionStatus> TaskStatus = InContext.GetTaskStatus(*this);
+	return TaskStatus.IsSet() && *TaskStatus == InExpectedStatus;
+}
+
 bool FSceneStateTask::ApplyBindings(const FSceneStateExecutionContext& InContext, FStructView InTaskInstance) const
 {
 	QUICK_SCOPE_CYCLE_COUNTER(SceneStateTask_ApplyBindings);
diff --git a/Engine/Plugins/VirtualProduction/SceneState/Source/SceneState/Public/SceneStateExecutionContext.h b/Engine/Plugins/VirtualProduction/SceneState/Source/SceneState/Public/SceneStateExecutionContext.h
--- a/Engine/Plugins/VirtualProduction/SceneState/Source/SceneState/Public/SceneStateExecutionContext.h
+++ b/Engine/Plugins/VirtualProduction/SceneState/Source/SceneState/Public/SceneStateExecutionContext.h
@@ -221,8 +221,8 @@ struct FSceneStateExecutionContext
 	/** Returns the Task Instance of this context for the given task absolute index. Returns an invalid view if not found */
 	FStructView FindTaskInstance(uint16 InTaskAbsoluteIndex) const;
 
-	/** Returns the given task's execution status */
-	UE_INTERNAL UE_API UE::SceneState::EExecutionStatus GetTaskStatus(const FSceneStateTask& InTask) const;
+	/** Returns the given task's execution status, or unset if task was not found */
+	UE_INTERNAL UE_API TOptional<UE::SceneState::EExecutionStatus> GetTaskStatus(const FSceneStateTask& InTask) const;
 
 	/** Set the task status of this context for the given task */
 	void SetTaskStatus(const FSceneStateTask& InTask, UE::SceneState::EExecutionStatus InStatus) const;
diff --git a/Engine/Plugins/VirtualProduction/SceneState/Source/SceneState/Public/Tasks/SceneStateTask.h b/Engine/Plugins/VirtualProduction/SceneState/Source/SceneState/Public/Tasks/SceneStateTask.h
--- a/Engine/Plugins/VirtualProduction/SceneState/Source/SceneState/Public/Tasks/SceneStateTask.h
+++ b/Engine/Plugins/VirtualProduction/SceneState/Source/SceneState/Public/Tasks/SceneStateTask.h
@@ -201,6 +201,9 @@ protected:
 	}
 
 private:
+	/** Returns true if the current task status on the given context matches the expected status */
+	bool IsTaskStatus(const FSceneStateExecutionContext& InContext, UE::SceneState::EExecutionStatus InExpectedStatus) const;
+
 	/** Applies Bindings to the given Task Instance */
 	bool ApplyBindings(const FSceneStateExecutionContext& InContext, FStructView InTaskInstance) const;
 
diff --git a/Engine/Plugins/VirtualProduction/SceneState/Source/SceneStateBlueprintEditor/Private/DebugExecutors/SceneStateDebugTaskExecutor.cpp b/Engine/Plugins/VirtualProduction/SceneState/Source/SceneStateBlueprintEditor/Private/DebugExecutors/SceneStateDebugTaskExecutor.cpp
--- a/Engine/Plugins/VirtualProduction/SceneState/Source/SceneStateBlueprintEditor/Private/DebugExecutors/SceneStateDebugTaskExecutor.cpp
+++ b/Engine/Plugins/VirtualProduction/SceneState/Source/SceneStateBlueprintEditor/Private/DebugExecutors/SceneStateDebugTaskExecutor.cpp
@@ -93,7 +93,8 @@ void FDebugTaskExecutor::OnExit(const FSceneStateExecutionContext& InExecutionCo
 
 void FDebugTaskExecutor::ConditionallyExit(const FSceneStateTask& InTask, const FSceneStateExecutionContext& InExecutionContext)
 {
-	if (InExecutionContext.GetTaskStatus(InTask) == EExecutionStatus::Finished)
+	const TOptional<UE::SceneState::EExecutionStatus> TaskStatus = InExecutionContext.GetTaskStatus(InTask);
+	if (TaskStatus.IsSet() && *TaskStatus == EExecutionStatus::Finished)
 	{
 		Exit();
 	}
