Motion Design Scene State: added an extra 'utility task' metadata + updating from deprecated api

#Synced-CL 51020850

diff --git a/Engine/Plugins/VirtualProduction/AvalancheSceneState/Source/AvalancheSceneState/Private/AvaSceneStateSchema.cpp b/Engine/Plugins/VirtualProduction/AvalancheSceneState/Source/AvalancheSceneState/Private/AvaSceneStateSchema.cpp
--- a/Engine/Plugins/VirtualProduction/AvalancheSceneState/Source/AvalancheSceneState/Private/AvaSceneStateSchema.cpp
+++ b/Engine/Plugins/VirtualProduction/AvalancheSceneState/Source/AvalancheSceneState/Private/AvaSceneStateSchema.cpp
@@ -11,11 +11,12 @@
 #if WITH_EDITOR
 bool UAvaSceneStateSchema::OnIsTaskStructAllowed(TSubScriptStructOf<FSceneStateTask> InTaskStruct) const
 {
-	// Allow Motion Design and Core tasks
+	// Allow Motion Design and Core/Utility tasks
 	constexpr bool bIncludeSuperStructs = false;
 	return UE::SceneState::Editor::StructHasAnyMetaData(InTaskStruct, 
 		{
 			UE::SceneState::Editor::MD_CoreTask,
+			UE::SceneState::Editor::MD_UtilityTask,
 			UE::AvaSceneState::Editor::MD_MotionDesignTask,
 		}, bIncludeSuperStructs);
 }
diff --git a/Engine/Plugins/VirtualProduction/SceneState/Source/SceneState/Private/SceneState.cpp b/Engine/Plugins/VirtualProduction/SceneState/Source/SceneState/Private/SceneState.cpp
--- a/Engine/Plugins/VirtualProduction/SceneState/Source/SceneState/Private/SceneState.cpp
+++ b/Engine/Plugins/VirtualProduction/SceneState/Source/SceneState/Private/SceneState.cpp
@@ -312,7 +312,7 @@ void FSceneState::InstanceTaskObjects(UObject* InOuter, TConstArrayView<FStructV
 						, *GetNameSafe(TemplateObject->GetOuter()));
 				}
 
-				void* const TargetPropertyMemory = UE::SceneState::ResolveVisitedPath(Target.GetScriptStruct(), Target.GetMemory(), InContext.Path);
+				void* const TargetPropertyMemory = PropertyVisitorHelpers::ResolveVisitedPath(Target.GetScriptStruct(), Target.GetMemory(), InContext.Path);
 				ObjectProperty->SetPropertyValue(TargetPropertyMemory, InstanceObject);
 				return EPropertyVisitorControlFlow::StepOver;
 			});
diff --git a/Engine/Plugins/VirtualProduction/SceneState/Source/SceneState/Private/SceneStateUtils.cpp b/Engine/Plugins/VirtualProduction/SceneState/Source/SceneState/Private/SceneStateUtils.cpp
--- a/Engine/Plugins/VirtualProduction/SceneState/Source/SceneState/Private/SceneStateUtils.cpp
+++ b/Engine/Plugins/VirtualProduction/SceneState/Source/SceneState/Private/SceneStateUtils.cpp
@@ -115,24 +115,7 @@ bool ApplyBatch(const FSceneStateExecutionContext& InContext, const FApplyBatchP
 
 void* ResolveVisitedPath(const UScriptStruct* InRootObject, void* InRootData, const FPropertyVisitorPath& InPath)
 {
-	if (!InRootObject || !InRootData)
-	{
-		return nullptr;
-	}
-
-	const TArray<FPropertyVisitorInfo>& PathArray = InPath.GetPath();
-	if (PathArray.IsEmpty())
-	{
-		return nullptr;
-	}
-
-	void* PropertyData = InRootObject->ResolveVisitedPathInfo(InRootData, PathArray[0]);
-	for (int32 PathIndex = 1; PropertyData && PathIndex < PathArray.Num(); ++PathIndex)
-	{
-		const FPropertyVisitorInfo& PreviousInfo = PathArray[PathIndex - 1];
-		PropertyData = PreviousInfo.Property->ResolveVisitedPathInfo(PropertyData, PathArray[PathIndex]);
-	}
-	return PropertyData;
+	return PropertyVisitorHelpers::ResolveVisitedPath(InRootObject, InRootData, InPath);
 }
 
 void DiscardObject(UObject* InObjectToDiscard)
diff --git a/Engine/Plugins/VirtualProduction/SceneState/Source/SceneState/Public/SceneStateUtils.h b/Engine/Plugins/VirtualProduction/SceneState/Source/SceneState/Public/SceneStateUtils.h
--- a/Engine/Plugins/VirtualProduction/SceneState/Source/SceneState/Public/SceneStateUtils.h
+++ b/Engine/Plugins/VirtualProduction/SceneState/Source/SceneState/Public/SceneStateUtils.h
@@ -57,12 +57,7 @@ namespace UE::SceneState
 	/** Applies a given binding batch to a target view within an execution context */
 	bool ApplyBatch(const FSceneStateExecutionContext& InContext, const FApplyBatchParams& InParams);
 
-	/**
-	 * Given a FPropertyVisitorPath, attempt to resolve that to a valid data pointer.
-	 * RootObject is required to implement ResolveVisitedPathInfo to provide the resolver logic.
-	 * TODO: FPropertyVisitorPath::GetPropertyDataPtr only takes a UObject ptr.
-	 * @see FPropertyVisitorPath::GetPropertyDataPtr
-	 */
+	UE_DEPRECATED(5.8, "Use PropertyVisitorHelpers::ResolveVisitedPath instead")
 	UE_API void* ResolveVisitedPath(const UScriptStruct* InRootObject, void* InRootData, const FPropertyVisitorPath& InPath);
 
 	/**
diff --git a/Engine/Plugins/VirtualProduction/SceneState/Source/SceneState/Public/Tasks/SceneStateMachineTask.h b/Engine/Plugins/VirtualProduction/SceneState/Source/SceneState/Public/Tasks/SceneStateMachineTask.h
--- a/Engine/Plugins/VirtualProduction/SceneState/Source/SceneState/Public/Tasks/SceneStateMachineTask.h
+++ b/Engine/Plugins/VirtualProduction/SceneState/Source/SceneState/Public/Tasks/SceneStateMachineTask.h
@@ -36,7 +36,7 @@ struct FSceneStateMachineTaskInstance : public FSceneStateTaskInstance
 };
 
 /** Task that runs a local state machine */
-USTRUCT(DisplayName="Run State Machine", Category="Core")
+USTRUCT(DisplayName="Run State Machine", Category="Core", meta=(CoreTask))
 struct FSceneStateMachineTask : public FSceneStateTask
 {
 	GENERATED_BODY()
diff --git a/Engine/Plugins/VirtualProduction/SceneState/Source/SceneState/Public/Tasks/SceneStatePlayerTask.h b/Engine/Plugins/VirtualProduction/SceneState/Source/SceneState/Public/Tasks/SceneStatePlayerTask.h
--- a/Engine/Plugins/VirtualProduction/SceneState/Source/SceneState/Public/Tasks/SceneStatePlayerTask.h
+++ b/Engine/Plugins/VirtualProduction/SceneState/Source/SceneState/Public/Tasks/SceneStatePlayerTask.h
@@ -19,7 +19,7 @@ struct FSceneStatePlayerTaskInstance : public FSceneStateTaskInstance
 };
 
 /** Task that runs a Scene State Object through a Player */
-USTRUCT(DisplayName="Run Player", Category="Core")
+USTRUCT(DisplayName="Run Player", Category="Core", meta=(CoreTask))
 struct FSceneStatePlayerTask : public FSceneStateTask
 {
 	GENERATED_BODY()
diff --git a/Engine/Plugins/VirtualProduction/SceneState/Source/SceneState/Public/Tasks/SceneStateTaskMetadata.h b/Engine/Plugins/VirtualProduction/SceneState/Source/SceneState/Public/Tasks/SceneStateTaskMetadata.h
--- a/Engine/Plugins/VirtualProduction/SceneState/Source/SceneState/Public/Tasks/SceneStateTaskMetadata.h
+++ b/Engine/Plugins/VirtualProduction/SceneState/Source/SceneState/Public/Tasks/SceneStateTaskMetadata.h
@@ -11,11 +11,17 @@ namespace UE::SceneState::Editor
 	const FLazyName MD_Hidden(TEXT("Hidden"));
 
 	/** 
-	 * The task is deemed as 'Core'. All the tasks that come with the Scene State plugin will be marked with this metadata. 
+	 * The task is deemed as 'Core'. These are tasks that should be allowed in most schemas. 
 	 * Derived tasks will NOT be recognized as 'CoreTask' unless they also specify this metadata.
 	 */
 	const FLazyName MD_CoreTask(TEXT("CoreTask"));
 
+	/**
+	 * The task is utility that can be useful for most schemas. All the non-core tasks that come with the Scene State plugin will be marked with this metadata. 
+	 * Derived tasks will NOT be recognized as 'UtilityTask' unless they also specify this metadata.
+	 */
+	const FLazyName MD_UtilityTask(TEXT("UtilityTask"));
+
 	/**
 	 * The task requires the context object to implement GetWorld().
 	 * This is just a convenience metadata to avoid having to call SupportsSchema.
diff --git a/Engine/Plugins/VirtualProduction/SceneState/Source/SceneStateBlueprint/Private/Actions/SceneStateBlueprintAction_Graph.cpp b/Engine/Plugins/VirtualProduction/SceneState/Source/SceneStateBlueprint/Private/Actions/SceneStateBlueprintAction_Graph.cpp
--- a/Engine/Plugins/VirtualProduction/SceneState/Source/SceneStateBlueprint/Private/Actions/SceneStateBlueprintAction_Graph.cpp
+++ b/Engine/Plugins/VirtualProduction/SceneState/Source/SceneStateBlueprint/Private/Actions/SceneStateBlueprintAction_Graph.cpp
@@ -35,7 +35,7 @@ FBlueprintAction_Graph::FBlueprintAction_Graph(const FArguments& InArgs)
 
 FName FBlueprintAction_Graph::StaticGetTypeId()
 {
-	static FName Type = TEXT("FEdGraphSchemaAction");
+	static FName Type = TEXT("UE::SceneState::Graph::FBlueprintAction_Graph");
 	return Type;
 }
 
diff --git a/Engine/Plugins/VirtualProduction/SceneState/Source/SceneStateBlueprint/Public/Actions/SceneStateBlueprintAction_Graph.h b/Engine/Plugins/VirtualProduction/SceneState/Source/SceneStateBlueprint/Public/Actions/SceneStateBlueprintAction_Graph.h
--- a/Engine/Plugins/VirtualProduction/SceneState/Source/SceneStateBlueprint/Public/Actions/SceneStateBlueprintAction_Graph.h
+++ b/Engine/Plugins/VirtualProduction/SceneState/Source/SceneStateBlueprint/Public/Actions/SceneStateBlueprintAction_Graph.h
@@ -10,7 +10,7 @@
 namespace UE::SceneState::Graph
 {
 
-struct FBlueprintAction_Graph : public FEdGraphSchemaAction_K2Graph
+struct UE_INTERNAL FBlueprintAction_Graph : public FEdGraphSchemaAction_K2Graph
 {
 	FBlueprintAction_Graph() = default;
 
diff --git a/Engine/Plugins/VirtualProduction/SceneState/Source/SceneStateBlueprintEditor/Private/Compiler/SceneStateBlueprintCompilerContext.cpp b/Engine/Plugins/VirtualProduction/SceneState/Source/SceneStateBlueprintEditor/Private/Compiler/SceneStateBlueprintCompilerContext.cpp
--- a/Engine/Plugins/VirtualProduction/SceneState/Source/SceneStateBlueprintEditor/Private/Compiler/SceneStateBlueprintCompilerContext.cpp
+++ b/Engine/Plugins/VirtualProduction/SceneState/Source/SceneStateBlueprintEditor/Private/Compiler/SceneStateBlueprintCompilerContext.cpp
@@ -75,11 +75,11 @@ const UClass* FBlueprintCompilerContext::GetGeneratedClass() const
 void FBlueprintCompilerContext::SpawnNewClass(const FString& InNewClassName)
 {
 	check(Blueprint);
-	NewGeneratedClass = FindObject<USceneStateGeneratedClass>(Blueprint->GetOutermost(), *InNewClassName);
+	NewGeneratedClass = FindObject<USceneStateGeneratedClass>(Blueprint->GetPackage(), *InNewClassName);
 
 	if (!NewGeneratedClass)
 	{
-		NewGeneratedClass = NewObject<USceneStateGeneratedClass>(Blueprint->GetOutermost(), *InNewClassName, RF_Public | RF_Transactional);
+		NewGeneratedClass = NewObject<USceneStateGeneratedClass>(Blueprint->GetPackage(), *InNewClassName, RF_Public | RF_Transactional);
 	}
 	else
 	{
diff --git a/Engine/Plugins/VirtualProduction/SceneState/Source/SceneStateGameplay/Private/SceneStateGameplaySchema.cpp b/Engine/Plugins/VirtualProduction/SceneState/Source/SceneStateGameplay/Private/SceneStateGameplaySchema.cpp
--- a/Engine/Plugins/VirtualProduction/SceneState/Source/SceneStateGameplay/Private/SceneStateGameplaySchema.cpp
+++ b/Engine/Plugins/VirtualProduction/SceneState/Source/SceneStateGameplay/Private/SceneStateGameplaySchema.cpp
@@ -12,11 +12,12 @@
 #if WITH_EDITOR
 bool USceneStateGameplaySchema::OnIsTaskStructAllowed(TSubScriptStructOf<FSceneStateTask> InTaskStruct) const
 {
-	// Only include core and gameplay tasks.
+	// Only include core/utility and gameplay tasks.
 	constexpr bool bIncludeSuperStructs = false;
 	return UE::SceneState::Editor::StructHasAnyMetaData(InTaskStruct,
 		{
 			UE::SceneState::Editor::MD_CoreTask,
+			UE::SceneState::Editor::MD_UtilityTask,
 			UE::SceneState::Editor::MD_GameplayTask,
 		}, bIncludeSuperStructs);
 }
diff --git a/Engine/Plugins/VirtualProduction/SceneState/Source/SceneStateTasks/Public/SceneStateDelayTask.h b/Engine/Plugins/VirtualProduction/SceneState/Source/SceneStateTasks/Public/SceneStateDelayTask.h
--- a/Engine/Plugins/VirtualProduction/SceneState/Source/SceneStateTasks/Public/SceneStateDelayTask.h
+++ b/Engine/Plugins/VirtualProduction/SceneState/Source/SceneStateTasks/Public/SceneStateDelayTask.h
@@ -20,7 +20,7 @@ struct FSceneStateDelayTaskInstance : public FSceneStateTaskInstance
 };
 
 /** Waits for a set amount of seconds */
-USTRUCT(DisplayName="Delay", Category="Core", meta=(CoreTask))
+USTRUCT(DisplayName="Delay", Category="Utility", meta=(UtilityTask))
 struct FSceneStateDelayTask : public FSceneStateTask
 {
 	GENERATED_BODY()
diff --git a/Engine/Plugins/VirtualProduction/SceneState/Source/SceneStateTasks/Public/SceneStatePrintStringTask.h b/Engine/Plugins/VirtualProduction/SceneState/Source/SceneStateTasks/Public/SceneStatePrintStringTask.h
--- a/Engine/Plugins/VirtualProduction/SceneState/Source/SceneStateTasks/Public/SceneStatePrintStringTask.h
+++ b/Engine/Plugins/VirtualProduction/SceneState/Source/SceneStateTasks/Public/SceneStatePrintStringTask.h
@@ -46,7 +46,7 @@ struct FSceneStatePrintStringTaskInstance : public FSceneStateTaskInstance
 	FSceneStatePrintSettings PrintSettings;
 };
 
-USTRUCT(DisplayName="Print String", Category="Core", meta=(CoreTask))
+USTRUCT(DisplayName="Print String", Category="Utility", meta=(UtilityTask))
 struct FSceneStatePrintStringTask : public FSceneStateTask
 {
 	GENERATED_BODY()
diff --git a/Engine/Plugins/VirtualProduction/SceneState/Source/SceneStateTasks/Public/SceneStateSpawnActorTask.h b/Engine/Plugins/VirtualProduction/SceneState/Source/SceneStateTasks/Public/SceneStateSpawnActorTask.h
--- a/Engine/Plugins/VirtualProduction/SceneState/Source/SceneStateTasks/Public/SceneStateSpawnActorTask.h
+++ b/Engine/Plugins/VirtualProduction/SceneState/Source/SceneStateTasks/Public/SceneStateSpawnActorTask.h
@@ -47,7 +47,7 @@ struct FSceneStateSpawnActorTaskInstance : public FSceneStateTaskInstance
 };
 
 /** Spawns an actor of a given class */
-USTRUCT(DisplayName="Spawn Actor", Category="Core", meta=(CoreTask, RequiresContextWorld))
+USTRUCT(DisplayName="Spawn Actor", Category="Utility", meta=(UtilityTask, RequiresContextWorld))
 struct FSceneStateSpawnActorTask : public FSceneStateTask
 {
 	GENERATED_BODY()
