Motion Design Scene State: fixed issue with pasted state names appending '_1'. Fixed issues with sub graphs not being visualized in hierarchy.

#Synced-CL 51020850

diff --git a/Engine/Plugins/VirtualProduction/SceneState/Source/SceneState/Public/SceneStateTemplateData.h b/Engine/Plugins/VirtualProduction/SceneState/Source/SceneState/Public/SceneStateTemplateData.h
--- a/Engine/Plugins/VirtualProduction/SceneState/Source/SceneState/Public/SceneStateTemplateData.h
+++ b/Engine/Plugins/VirtualProduction/SceneState/Source/SceneState/Public/SceneStateTemplateData.h
@@ -155,7 +155,7 @@ public:
 	UPROPERTY()
 	TArray<FSceneStateMachine> StateMachines;
 
-	/** Map of the Top-Level State machine Parameters id to the index in the state machine array */
+	/** Map of the state machine parameters id to the index in the state machine array */
 	UPROPERTY()
 	TMap<FGuid, uint16> StateMachineIdToIndex;
 
diff --git a/Engine/Plugins/VirtualProduction/SceneState/Source/SceneStateBlueprint/Public/Actions/SceneStateBlueprintAction_Graph.h b/Engine/Plugins/VirtualProduction/SceneState/Source/SceneStateBlueprint/Public/Actions/SceneStateBlueprintAction_Graph.h
--- a/Engine/Plugins/VirtualProduction/SceneState/Source/SceneStateBlueprint/Public/Actions/SceneStateBlueprintAction_Graph.h
+++ b/Engine/Plugins/VirtualProduction/SceneState/Source/SceneStateBlueprint/Public/Actions/SceneStateBlueprintAction_Graph.h
@@ -39,6 +39,6 @@ private:
 	FSlateIcon Icon;
 };
 
-} // UE::SceneState::Editor
+} // UE::SceneState::Graph
 
 #undef UE_API
diff --git a/Engine/Plugins/VirtualProduction/SceneState/Source/SceneStateBlueprintEditor/Private/DetailsView/SceneStateMachineGraphCustomization.cpp b/Engine/Plugins/VirtualProduction/SceneState/Source/SceneStateBlueprintEditor/Private/DetailsView/SceneStateMachineGraphCustomization.cpp
--- a/Engine/Plugins/VirtualProduction/SceneState/Source/SceneStateBlueprintEditor/Private/DetailsView/SceneStateMachineGraphCustomization.cpp
+++ b/Engine/Plugins/VirtualProduction/SceneState/Source/SceneStateBlueprintEditor/Private/DetailsView/SceneStateMachineGraphCustomization.cpp
@@ -45,19 +45,20 @@ void FStateMachineGraphCustomization::CustomizeDetails(IDetailLayoutBuilder& InD
 	IDetailCategoryBuilder& StateMachineCategory = InDetailBuilder.EditCategory(TEXT("State Machine"));
 
 	// Show category if there's at least one top-level graph
-	const bool bShowCategory = Graphs.ContainsByPredicate(
+	const bool bHasTopLevelGraph = Graphs.ContainsByPredicate(
 		[](TWeakObjectPtr<USceneStateMachineGraph> InGraphWeak)
 		{
 			USceneStateMachineGraph* Graph = InGraphWeak.Get();
 			return Graph && Graph->IsRootStateMachine();
 		});
 
-	if (bShowCategory)
+	if (bHasTopLevelGraph)
 	{
 		StateMachineCategory.AddProperty(CategoryHandle);
-	}
 
-	StateMachineCategory.AddProperty(RunModeHandle);
+		// Todo: at the moment run mode is not supported in sub-graphs
+		StateMachineCategory.AddProperty(RunModeHandle);
+	}
 
 	FGuid StateMachineId;
 	GetGuid(StateMachineIdHandle, StateMachineId);
diff --git a/Engine/Plugins/VirtualProduction/SceneState/Source/SceneStateBlueprintEditor/Private/DetailsView/Widgets/SSceneStateMachinePicker.cpp b/Engine/Plugins/VirtualProduction/SceneState/Source/SceneStateBlueprintEditor/Private/DetailsView/Widgets/SSceneStateMachinePicker.cpp
--- a/Engine/Plugins/VirtualProduction/SceneState/Source/SceneStateBlueprintEditor/Private/DetailsView/Widgets/SSceneStateMachinePicker.cpp
+++ b/Engine/Plugins/VirtualProduction/SceneState/Source/SceneStateBlueprintEditor/Private/DetailsView/Widgets/SSceneStateMachinePicker.cpp
@@ -55,9 +55,16 @@ void SStateMachinePicker::Construct(const FArguments& InArgs, const TSharedRef<I
 	StateMachineIdHandle = InStateMachineIdHandle;
 	ParametersHandle = InParametersHandle;
 
-	USceneStateMachineGraph* StateMachineGraph = Private::GetTypedOuter<USceneStateMachineGraph>(ParametersHandle);
-	check(StateMachineGraph);
-	OwningStateMachineId = StateMachineGraph->ParametersId;
+	// Add all the ids of outer state machine graphs to the ignore set
+	StateMachineIdsToIgnore.Reset();
+	{
+		const USceneStateMachineGraph* StateMachineGraph = Private::GetTypedOuter<USceneStateMachineGraph>(ParametersHandle);
+		while (StateMachineGraph)
+		{
+			StateMachineIdsToIgnore.Add(StateMachineGraph->ParametersId);
+			StateMachineGraph = StateMachineGraph->GetTypedOuter<USceneStateMachineGraph>();
+		}
+	}
 
 	OnParametersChangedHandle = USceneStateMachineGraph::OnParametersChanged().AddSP(this
 		, &SStateMachinePicker::OnParametersChanged);
@@ -137,8 +144,8 @@ void SStateMachinePicker::RefreshOptions()
 	{
 		USceneStateMachineGraph* StateMachineGraph = Cast<USceneStateMachineGraph>(Graph);
 
-		// Skip the state machine that owns this picker
-		if (!StateMachineGraph || StateMachineGraph->ParametersId == OwningStateMachineId)
+		// Skip the state machines in the ignore set
+		if (!StateMachineGraph || StateMachineIdsToIgnore.Contains(StateMachineGraph->ParametersId))
 		{
 			continue;
 		}
diff --git a/Engine/Plugins/VirtualProduction/SceneState/Source/SceneStateBlueprintEditor/Private/DetailsView/Widgets/SSceneStateMachinePicker.h b/Engine/Plugins/VirtualProduction/SceneState/Source/SceneStateBlueprintEditor/Private/DetailsView/Widgets/SSceneStateMachinePicker.h
--- a/Engine/Plugins/VirtualProduction/SceneState/Source/SceneStateBlueprintEditor/Private/DetailsView/Widgets/SSceneStateMachinePicker.h
+++ b/Engine/Plugins/VirtualProduction/SceneState/Source/SceneStateBlueprintEditor/Private/DetailsView/Widgets/SSceneStateMachinePicker.h
@@ -59,8 +59,8 @@ private:
 	TSharedPtr<FStateMachinePickerOption> SelectedOption;
 	FDelegateHandle OnParametersChangedHandle;
 
-	/** The State Machine Id of the State Machine this picker is outered to */
-	FGuid OwningStateMachineId;
+	/** State machine ids that this picker should not display as options */
+	TSet<FGuid> StateMachineIdsToIgnore;
 };
 
 } // UE::SceneState::Editor
diff --git a/Engine/Plugins/VirtualProduction/SceneState/Source/SceneStateBlueprintEditor/Private/Widgets/SSceneStateStateMachineMenu.cpp b/Engine/Plugins/VirtualProduction/SceneState/Source/SceneStateBlueprintEditor/Private/Widgets/SSceneStateStateMachineMenu.cpp
--- a/Engine/Plugins/VirtualProduction/SceneState/Source/SceneStateBlueprintEditor/Private/Widgets/SSceneStateStateMachineMenu.cpp
+++ b/Engine/Plugins/VirtualProduction/SceneState/Source/SceneStateBlueprintEditor/Private/Widgets/SSceneStateStateMachineMenu.cpp
@@ -2,10 +2,12 @@
 
 #include "SSceneStateStateMachineMenu.h"
 #include "Actions/SceneStateBlueprintAction_Graph.h"
+#include "Actions/SceneStateMachineNodeAction.h"
 #include "GraphActionNode.h"
 #include "Kismet2/BlueprintEditorUtils.h"
 #include "Menus/SceneStateMachineAddMenu.h"
 #include "Menus/SceneStateMachineContextMenu.h"
+#include "Nodes/SceneStateMachineNode.h"
 #include "SGraphPalette.h"
 #include "SPositiveActionButton.h"
 #include "SSceneStateBlueprintPaletteItem.h"
@@ -226,6 +228,15 @@ void SStateMachineMenu::GetGraphActionDetails(const TSharedPtr<FEdGraphSchemaAct
 			OutDetailsText = DisplayInfo.PlainName;
 		}
 	}
+	else if (GraphActionType == Graph::FStateMachineAction_Node::StaticGetTypeId())
+	{
+		TSharedPtr<Graph::FStateMachineAction_Node> NodeAction = StaticCastSharedPtr<Graph::FStateMachineAction_Node>(InAction);
+		if (USceneStateMachineNode* Node = NodeAction->GetNode())
+		{
+			OutDetailsObject = Node;
+			OutDetailsText = FText::FromName(Node->GetNodeName());
+		}
+	}
 }
 
 void SStateMachineMenu::OnGraphActionSelected(const TArray<TSharedPtr<FEdGraphSchemaAction>>& InActions, ESelectInfo::Type InSelectionType)
@@ -268,6 +279,40 @@ FReply SStateMachineMenu::OnGraphActionDragged(const TArray<TSharedPtr<FEdGraphS
 	return FReply::Unhandled();
 }
 
+bool SStateMachineMenu::ShouldProcessNode(TNotNull<const USceneStateMachineNode*> InGraphNode) const
+{
+	return InGraphNode->GetBoundGraphs().ContainsByPredicate(
+		[This=this](UEdGraph* InBoundGraph)
+		{
+			return This->ShouldProcessGraph(InBoundGraph);
+		});
+}
+
+void SStateMachineMenu::CollectNodeActionsRecursive(TNotNull<USceneStateMachineNode*> InNode, FText InCategory, FGraphActionListBuilderBase& OutActions)
+{
+	if (!ShouldProcessNode(InNode))
+	{
+		return;
+	}
+
+	const Graph::FStateMachineAction_Node::FArguments Arguments
+		{
+			.Node = InNode,
+			.Category = InCategory,
+			.SectionID = NodeSectionID::GRAPH,
+		};
+
+	TSharedRef<Graph::FStateMachineAction_Node> NodeAction = MakeShared<Graph::FStateMachineAction_Node>(Arguments);
+	OutActions.AddAction(NodeAction);
+
+	InCategory = FText::Format(INVTEXT("{0}|{1}"), InCategory, NodeAction->GetMenuDescription());
+
+	for (UEdGraph* BoundGraph : InNode->GetBoundGraphs())
+	{
+		CollectGraphActionsRecursive(BoundGraph, InCategory, EEdGraphSchemaAction_K2Graph::Type::Subgraph, OutActions);
+	}
+}
+
 bool SStateMachineMenu::ShouldProcessGraph(UEdGraph* InGraph) const
 {
 	if (!InGraph)
@@ -334,18 +379,14 @@ void SStateMachineMenu::CollectGraphActionsRecursive(UEdGraph* InGraph, FText In
 	TSharedRef<Graph::FBlueprintAction_Graph> GraphAction = MakeShared<Graph::FBlueprintAction_Graph>(Arguments);
 	OutActions.AddAction(GraphAction);
 
-	if (InCategory.IsEmpty())
-	{
-		InCategory = GraphAction->GetCategory();
-	}
-	else
-	{
-		InCategory = FText::Format(INVTEXT("{0}|{1}"), InCategory, GraphAction->GetMenuDescription());
-	}
+	InCategory = FText::Format(INVTEXT("{0}|{1}"), InCategory, GraphAction->GetMenuDescription());
 
-	for (UEdGraph* SubGraph : InGraph->SubGraphs)
+	for (UEdGraphNode* GraphNode : InGraph->Nodes)
 	{
-		CollectGraphActionsRecursive(SubGraph, InCategory, EEdGraphSchemaAction_K2Graph::Type::Subgraph, OutActions);
+		if (USceneStateMachineNode* Node = Cast<USceneStateMachineNode>(GraphNode))
+		{
+			CollectNodeActionsRecursive(Node, InCategory, OutActions);
+		}
 	}
 }
 
@@ -418,7 +459,8 @@ void SStateMachineMenu::ExecuteGraphAction(const TSharedPtr<FEdGraphSchemaAction
 		return;
 	}
 
-	if (InAction->GetTypeId() == Graph::FBlueprintAction_Graph::StaticGetTypeId())
+	const FName ActionTypeId = InAction->GetTypeId();
+	if (ActionTypeId == Graph::FBlueprintAction_Graph::StaticGetTypeId())
 	{
 		TSharedPtr<Graph::FBlueprintAction_Graph> GraphAction = StaticCastSharedPtr<Graph::FBlueprintAction_Graph>(InAction);
 		if (GraphAction->EdGraph)
@@ -426,6 +468,14 @@ void SStateMachineMenu::ExecuteGraphAction(const TSharedPtr<FEdGraphSchemaAction
 			BlueprintEditor->JumpToHyperlink(GraphAction->EdGraph, /*bRequestRename*/false);
 		}
 	}
+	else if (ActionTypeId == Graph::FStateMachineAction_Node::StaticGetTypeId())
+	{
+		TSharedPtr<Graph::FStateMachineAction_Node> NodeAction = StaticCastSharedPtr<Graph::FStateMachineAction_Node>(InAction);
+		if (USceneStateMachineNode* Node = NodeAction->GetNode())
+		{
+			BlueprintEditor->JumpToHyperlink(Node, /*bRequestRename*/false);
+		}
+	}
 }
 
 bool SStateMachineMenu::CanRequestRenameOnActionNode(TWeakPtr<FGraphActionNode> InSelectedNodeWeak) const
diff --git a/Engine/Plugins/VirtualProduction/SceneState/Source/SceneStateBlueprintEditor/Private/Widgets/SSceneStateStateMachineMenu.h b/Engine/Plugins/VirtualProduction/SceneState/Source/SceneStateBlueprintEditor/Private/Widgets/SSceneStateStateMachineMenu.h
--- a/Engine/Plugins/VirtualProduction/SceneState/Source/SceneStateBlueprintEditor/Private/Widgets/SSceneStateStateMachineMenu.h
+++ b/Engine/Plugins/VirtualProduction/SceneState/Source/SceneStateBlueprintEditor/Private/Widgets/SSceneStateStateMachineMenu.h
@@ -11,6 +11,7 @@ class SSearchBox;
 class UBlueprint;
 class UEdGraph;
 class UObject;
+class USceneStateMachineNode;
 struct FCreateWidgetForActionData;
 struct FEdGraphSchemaAction;
 struct FGraphActionListBuilderBase;
@@ -65,8 +66,10 @@ private:
 	/** Called when dragging the graph action entry */
 	FReply OnGraphActionDragged(const TArray<TSharedPtr<FEdGraphSchemaAction>>& InActions, const FPointerEvent& InMouseEvent);
 
-	bool ShouldProcessGraph(UEdGraph* InGraph) const;
+	bool ShouldProcessNode(TNotNull<const USceneStateMachineNode*> InGraphNode) const;
+	void CollectNodeActionsRecursive(TNotNull<USceneStateMachineNode*> InNode, FText InCategory, FGraphActionListBuilderBase& OutActions);
 
+	bool ShouldProcessGraph(UEdGraph* InGraph) const;
 	void CollectGraphActionsRecursive(UEdGraph* InGraph, FText InCategory, int32 InGraphType, FGraphActionListBuilderBase& OutActions);
 
 	void CollectGraphActions(FGraphActionListBuilderBase& OutActions);
diff --git a/Engine/Plugins/VirtualProduction/SceneState/Source/SceneStateMachineEditor/Private/SceneStateMachineCompiler.cpp b/Engine/Plugins/VirtualProduction/SceneState/Source/SceneStateMachineEditor/Private/SceneStateMachineCompiler.cpp
--- a/Engine/Plugins/VirtualProduction/SceneState/Source/SceneStateMachineEditor/Private/SceneStateMachineCompiler.cpp
+++ b/Engine/Plugins/VirtualProduction/SceneState/Source/SceneStateMachineEditor/Private/SceneStateMachineCompiler.cpp
@@ -203,7 +203,7 @@ void FStateMachineCompiler::FollowTransitions(TConstArrayView<USceneStateMachine
 void FStateMachineCompiler::CompileSubStateMachines(FSceneState& InNewState, const USceneStateMachineStateNode* InStateNode)
 {
 	// Compile Sub State Machine Graphs into Runtime State Machines owned by this new State
-	TArray<UEdGraph*> SubGraphs = InStateNode->GetSubGraphs();
+	const TArray<UEdGraph*> SubGraphs = InStateNode->GetSubGraphs();
 
 	TArray<FSceneStateMachine> SubStateMachines;
 	SubStateMachines.Reserve(SubGraphs.Num());
@@ -211,6 +211,9 @@ void FStateMachineCompiler::CompileSubStateMachines(FSceneState& InNewState, con
 	TMap<FObjectKey, uint16> StateMachineGraphToIndex;
 	StateMachineGraphToIndex.Reserve(SubGraphs.Num());
 
+	TMap<FGuid, uint16> StateMachineIdToIndex;
+	StateMachineIdToIndex.Reserve(SubGraphs.Num());
+
 	for (UEdGraph* SubGraph : SubGraphs)
 	{
 		USceneStateMachineGraph* NewStateMachineGraph = Cast<USceneStateMachineGraph>(SubGraph);
@@ -223,8 +226,10 @@ void FStateMachineCompiler::CompileSubStateMachines(FSceneState& InNewState, con
 		FSceneStateMachine NewStateMachine = StateMachineCompiler.Compile();
 		if (NewStateMachine.IsValid())
 		{
-			StateMachineGraphToIndex.Emplace(NewStateMachineGraph, StateMachineGraphToIndex.Num());
+			const uint16 StateMachineIndex = SubStateMachines.Num();
 			SubStateMachines.Emplace(MoveTemp(NewStateMachine));
+			StateMachineGraphToIndex.Emplace(NewStateMachineGraph, StateMachineIndex);
+			StateMachineIdToIndex.Emplace(NewStateMachineGraph->ParametersId, StateMachineIndex);
 		}
 	}
 
@@ -235,6 +240,10 @@ void FStateMachineCompiler::CompileSubStateMachines(FSceneState& InNewState, con
 	// Upgrade the map to absolute indices before baking it to the generated class' map
 	ToAbsoluteIndexMap(StateMachineGraphToIndex, InNewState.StateMachineRange.Index);
 	TemplateData->StateMachineGraphToIndex.Append(MoveTemp(StateMachineGraphToIndex));
+
+	// Upgrade the map to absolute indices before baking it to the generated class' map
+	ToAbsoluteIndexMap(StateMachineIdToIndex, InNewState.StateMachineRange.Index);
+	TemplateData->StateMachineIdToIndex.Append(MoveTemp(StateMachineIdToIndex));
 }
 
 void FStateMachineCompiler::AddConduitToCompile(const USceneStateMachineConduitNode* InConduitNode)
@@ -283,7 +292,6 @@ void FStateMachineCompiler::CompileTasks(FSceneState& InNewState, const USceneSt
 		Task.GetMutable<FSceneStateTask>().PrerequisiteRange.Index += TaskPrerequisiteIndex;
 	}
 
-	// todo: bigger chunks to reduce FInstancedStructContainer realloc
 	TemplateData->Tasks.Append(MoveTemp(TaskCompileResult.Tasks));
 	TemplateData->TaskInstances.Append(TaskCompileResult.TaskInstances);
 	TemplateData->TaskMetadata.Append(MoveTemp(TaskCompileResult.TaskMetadata));
diff --git a/Engine/Plugins/VirtualProduction/SceneState/Source/SceneStateMachineGraph/Private/Actions/SceneStateMachineNodeAction.cpp b/Engine/Plugins/VirtualProduction/SceneState/Source/SceneStateMachineGraph/Private/Actions/SceneStateMachineNodeAction.cpp
new file mode 100644
--- /dev/null
+++ b/Engine/Plugins/VirtualProduction/SceneState/Source/SceneStateMachineGraph/Private/Actions/SceneStateMachineNodeAction.cpp
@@ -0,0 +1,54 @@
+// Copyright Epic Games, Inc. All Rights Reserved.
+
+#include "Actions/SceneStateMachineNodeAction.h"
+#include "Nodes/SceneStateMachineNode.h"
+
+#define LOCTEXT_NAMESPACE "SceneStateBlueprintAction_StateNode"
+
+namespace UE::SceneState::Graph
+{
+
+FStateMachineAction_Node::FStateMachineAction_Node(const FArguments& InArgs)
+{
+	NodeWeak = InArgs.Node;
+	SectionID = InArgs.SectionID;
+
+	if (InArgs.Node)
+	{
+		UpdateSearchData(FText::FromName(InArgs.Node->GetNodeName()), FText(),  InArgs.Category, FText());
+	}
+}
+
+FName FStateMachineAction_Node::StaticGetTypeId()
+{
+	static FName Type = TEXT("UE::SceneState::Graph::FBlueprintAction_StateMachineNode");
+	return Type;
+}
+
+bool FStateMachineAction_Node::IsParentable() const
+{
+	return true;
+}
+
+bool FStateMachineAction_Node::CanBeRenamed() const
+{
+	return false;
+}
+
+FName FStateMachineAction_Node::GetTypeId() const
+{
+	return StaticGetTypeId();
+}
+
+FEdGraphSchemaActionDefiningObject FStateMachineAction_Node::GetPersistentItemDefiningObject() const
+{
+	if (USceneStateMachineNode* Node = NodeWeak.Get())
+	{
+		return FEdGraphSchemaActionDefiningObject(Node->GetOuter());
+	}
+	return FEdGraphSchemaActionDefiningObject(nullptr);
+}
+
+} // UE::SceneState::Graph
+
+#undef LOCTEXT_NAMESPACE
diff --git a/Engine/Plugins/VirtualProduction/SceneState/Source/SceneStateMachineGraph/Private/Nodes/SceneStateMachineNode.cpp b/Engine/Plugins/VirtualProduction/SceneState/Source/SceneStateMachineGraph/Private/Nodes/SceneStateMachineNode.cpp
--- a/Engine/Plugins/VirtualProduction/SceneState/Source/SceneStateMachineGraph/Private/Nodes/SceneStateMachineNode.cpp
+++ b/Engine/Plugins/VirtualProduction/SceneState/Source/SceneStateMachineGraph/Private/Nodes/SceneStateMachineNode.cpp
@@ -234,7 +234,7 @@ void USceneStateMachineNode::GenerateNewNodeName()
 {
 	const TSharedRef<INameValidatorInterface> NameValidator = FNameValidatorFactory::MakeValidator(this).ToSharedRef();
 
-	FString NewNodeName = NodeName.ToString();
+	FString NewNodeName = NodeName.GetPlainNameString();
 	NameValidator->FindValidString(NewNodeName);
 	NodeName = *NewNodeName;
 }
diff --git a/Engine/Plugins/VirtualProduction/SceneState/Source/SceneStateMachineGraph/Private/SceneStateMachineGraphSchema.cpp b/Engine/Plugins/VirtualProduction/SceneState/Source/SceneStateMachineGraph/Private/SceneStateMachineGraphSchema.cpp
--- a/Engine/Plugins/VirtualProduction/SceneState/Source/SceneStateMachineGraph/Private/SceneStateMachineGraphSchema.cpp
+++ b/Engine/Plugins/VirtualProduction/SceneState/Source/SceneStateMachineGraph/Private/SceneStateMachineGraphSchema.cpp
@@ -403,7 +403,7 @@ void USceneStateMachineGraphSchema::GetGraphDisplayInformation(const UEdGraph& I
 {
 	OutDisplayInfo.PlainName = FText::FromName(InGraph.GetFName());
 	OutDisplayInfo.DisplayName = OutDisplayInfo.PlainName;
-	OutDisplayInfo.Tooltip = LOCTEXT("GraphTooltip", "Graph used to transition between different states");
+	OutDisplayInfo.Tooltip = InGraph.GetClass()->GetToolTipText();
 }
 
 void USceneStateMachineGraphSchema::BreakNodeLinks(UEdGraphNode& InTargetNode) const
diff --git a/Engine/Plugins/VirtualProduction/SceneState/Source/SceneStateMachineGraph/Public/Actions/SceneStateMachineNodeAction.h b/Engine/Plugins/VirtualProduction/SceneState/Source/SceneStateMachineGraph/Public/Actions/SceneStateMachineNodeAction.h
new file mode 100644
--- /dev/null
+++ b/Engine/Plugins/VirtualProduction/SceneState/Source/SceneStateMachineGraph/Public/Actions/SceneStateMachineNodeAction.h
@@ -0,0 +1,47 @@
+// Copyright Epic Games, Inc. All Rights Reserved.
+
+#pragma once
+
+#include "EdGraphSchema_K2_Actions.h"
+#include "Textures/SlateIcon.h"
+
+#define UE_API SCENESTATEMACHINEGRAPH_API
+
+class USceneStateMachineNode;
+
+namespace UE::SceneState::Graph
+{
+
+struct UE_INTERNAL FStateMachineAction_Node : public FEdGraphSchemaAction
+{
+	FStateMachineAction_Node() = default;
+
+	struct FArguments
+	{
+		USceneStateMachineNode* Node = nullptr;
+		FText Category;
+		int32 SectionID;
+	};
+	UE_API explicit FStateMachineAction_Node(const FArguments& InArgs);
+
+	UE_API static FName StaticGetTypeId();
+
+	USceneStateMachineNode* GetNode() const
+	{
+		return NodeWeak.Get();
+	}
+
+	//~ Begin FEdGraphSchemaAction
+	UE_API virtual bool IsParentable() const override;
+	UE_API virtual bool CanBeRenamed() const override;
+	UE_API virtual FName GetTypeId() const override;
+	UE_API virtual FEdGraphSchemaActionDefiningObject GetPersistentItemDefiningObject() const override;
+	//~ End FEdGraphSchemaAction
+
+private:
+	TWeakObjectPtr<USceneStateMachineNode> NodeWeak;
+};
+
+} // UE::SceneState::Graph
+
+#undef UE_API
