Motion Design Scene State: fix navigating nodes

#Synced-CL 51020850

diff --git a/Engine/Plugins/VirtualProduction/SceneState/Source/SceneStateBlueprintEditor/Private/SceneStateBlueprintEditor.cpp b/Engine/Plugins/VirtualProduction/SceneState/Source/SceneStateBlueprintEditor/Private/SceneStateBlueprintEditor.cpp
--- a/Engine/Plugins/VirtualProduction/SceneState/Source/SceneStateBlueprintEditor/Private/SceneStateBlueprintEditor.cpp
+++ b/Engine/Plugins/VirtualProduction/SceneState/Source/SceneStateBlueprintEditor/Private/SceneStateBlueprintEditor.cpp
@@ -114,8 +114,23 @@ void FSceneStateBlueprintEditor::RefreshMyBlueprint()
 
 void FSceneStateBlueprintEditor::JumpToHyperlink(const UObject* InObjectReference, bool bInRequestRename)
 {
-	if (Cast<USceneStateMachineGraph>(InObjectReference) || Cast<USceneStateTransitionGraph>(InObjectReference))
+	if (!InObjectReference)
 	{
+		return;
+	}
+
+	const bool bIsStateMachineObject = InObjectReference->IsA<USceneStateMachineGraph>()
+		|| InObjectReference->IsA<USceneStateTransitionGraph>()
+		|| InObjectReference->IsA<USceneStateMachineNode>();
+
+	if (bIsStateMachineObject)
+	{
+		const USceneStateMachineNode* const TargetNode = Cast<USceneStateMachineNode>(InObjectReference);
+		if (TargetNode)
+		{
+			InObjectReference = TargetNode->GetTypedOuter<UEdGraph>();
+		}
+
 		// Default to navigate the same document
 		FDocumentTracker::EOpenDocumentCause OpenMode = FDocumentTracker::NavigatingCurrentDocument;
 
@@ -126,6 +141,11 @@ void FSceneStateBlueprintEditor::JumpToHyperlink(const UObject* InObjectReferenc
 		}
 
 		OpenDocument(InObjectReference, OpenMode);
+
+		if (TargetNode)
+		{
+			JumpToNode(TargetNode);
+		}
 	}
 	else
 	{
