Motion Design Scene State: binding compiler to happen later in the compilation stage to when the cdo can be safely retrieved and pass valid binding data views to the binding collection.

#Synced-CL 51290016

#Synced-CL 51290016

diff --git a/Engine/Plugins/VirtualProduction/SceneState/Source/SceneState/Public/SceneStateGeneratedClass.h b/Engine/Plugins/VirtualProduction/SceneState/Source/SceneState/Public/SceneStateGeneratedClass.h
--- a/Engine/Plugins/VirtualProduction/SceneState/Source/SceneState/Public/SceneStateGeneratedClass.h
+++ b/Engine/Plugins/VirtualProduction/SceneState/Source/SceneState/Public/SceneStateGeneratedClass.h
@@ -49,9 +49,10 @@ public:
 
 private:
 	/** Called to resolve the template data bindings */
-	UE_API void ResolveBindings();
+	UE_INTERNAL UE_API void ResolveBindings();
 
-	bool IsFullClass() const;
+	/** Returns true if this class is a 'full' class (i.e. not a SKEL_ or REINST_ class) */
+	UE_INTERNAL UE_API bool IsFullClass() const;
 
 #if WITH_EDITOR
 	void OnObjectsReinstanced(const TMap<UObject*, UObject*>& InReplacementMap);
diff --git a/Engine/Plugins/VirtualProduction/SceneState/Source/SceneStateBlueprintEditor/Private/Compiler/SceneStateBindingCompiler.cpp b/Engine/Plugins/VirtualProduction/SceneState/Source/SceneStateBlueprintEditor/Private/Compiler/SceneStateBindingCompiler.cpp
--- a/Engine/Plugins/VirtualProduction/SceneState/Source/SceneStateBlueprintEditor/Private/Compiler/SceneStateBindingCompiler.cpp
+++ b/Engine/Plugins/VirtualProduction/SceneState/Source/SceneStateBlueprintEditor/Private/Compiler/SceneStateBindingCompiler.cpp
@@ -103,11 +103,10 @@ void FBindingCompiler::AddDataView(const FSceneStateBindingDataHandle& InDataHan
 	TemplateDataViewMap.Add(InDataHandle, InDataView);
 }
 
-void FBindingCompiler::AddBindingDesc(FSceneStateBindingDesc&& InBindingDesc)
+void FBindingCompiler::AddBindingDesc(FSceneStateBindingDesc&& InBindingDesc, FPropertyBindingDataView InDataView)
 {
 	// Record this binding in the valid binding map. This will be used to then remove all the invalid bindings from the collection
-	// The binding data view can be filled in with null memory, as only the struct is required for resolution
-	ValidBindingMap.Add(InBindingDesc.ID, FPropertyBindingDataView(InBindingDesc.Struct, nullptr));
+	ValidBindingMap.Add(InBindingDesc.ID, InDataView);
 
 	TemplateData->BindingCollection.BindingDescs.Add(MoveTemp(InBindingDesc));
 }
@@ -120,7 +119,8 @@ void FBindingCompiler::AddRootBindingDesc()
 	BindingDesc.Struct = Context.GetGeneratedClass();
 	BindingDesc.DataHandle = FSceneStateBindingDataHandle(ESceneStateDataType::Root);
 
-	AddBindingDesc(MoveTemp(BindingDesc));
+	// Binding compilation happens late in the blueprint compilation, where the cdo is safe to get
+	AddBindingDesc(MoveTemp(BindingDesc), FPropertyBindingDataView(Blueprint->GeneratedClass->GetDefaultObject()));
 }
 
 void FBindingCompiler::AddStateBindingDescs()
@@ -129,7 +129,7 @@ void FBindingCompiler::AddStateBindingDescs()
 	{
 		const uint16 StateIndex = Pair.Value;
 
-		const FSceneState& State = TemplateData->States[StateIndex];
+		FSceneState& State = TemplateData->States[StateIndex];
 		if (!State.Parameters.IsValid())
 		{
 			continue;
@@ -137,11 +137,11 @@ void FBindingCompiler::AddStateBindingDescs()
 
 		FSceneStateBindingDesc BindingDesc;
 		BindingDesc.ID = Pair.Key;
-		BindingDesc.Name = TEXT("State Parameters");
+		BindingDesc.Name = FName("StateParameters", StateIndex);
 		BindingDesc.Struct = State.Parameters.GetPropertyBagStruct();
 		BindingDesc.DataHandle = FSceneStateBindingDataHandle(ESceneStateDataType::State, StateIndex);
 
-		AddBindingDesc(MoveTemp(BindingDesc));
+		AddBindingDesc(MoveTemp(BindingDesc), State.Parameters.GetMutableValue());
 	}
 }
 
@@ -151,7 +151,7 @@ void FBindingCompiler::AddStateMachineBindingDescs()
 	{
 		const uint16 StateMachineIndex = Pair.Value;
 
-		const FSceneStateMachine& StateMachine = TemplateData->StateMachines[StateMachineIndex];
+		FSceneStateMachine& StateMachine = TemplateData->StateMachines[StateMachineIndex];
 		if (!StateMachine.Parameters.IsValid())
 		{
 			continue;
@@ -159,17 +159,17 @@ void FBindingCompiler::AddStateMachineBindingDescs()
 
 		FSceneStateBindingDesc BindingDesc;
 		BindingDesc.ID = Pair.Key;
-		BindingDesc.Name = TEXT("State Machine Parameters");
+		BindingDesc.Name = FName("StateMachineParameters", StateMachineIndex);
 		BindingDesc.Struct = StateMachine.Parameters.GetPropertyBagStruct();
 		BindingDesc.DataHandle = FSceneStateBindingDataHandle(ESceneStateDataType::StateMachine, StateMachineIndex);
 
-		AddBindingDesc(MoveTemp(BindingDesc));
+		AddBindingDesc(MoveTemp(BindingDesc), StateMachine.Parameters.GetMutableValue());
 	}
 }
 
 void FBindingCompiler::AddTransitionBindingDescs()
 {
-	for (const TPair<uint16, FInstancedPropertyBag>& Pair : TemplateData->TransitionParameters)
+	for (TPair<uint16, FInstancedPropertyBag>& Pair : TemplateData->TransitionParameters)
 	{
 		const uint16 TransitionIndex = Pair.Key;
 
@@ -177,11 +177,11 @@ void FBindingCompiler::AddTransitionBindingDescs()
 
 		FSceneStateBindingDesc BindingDesc;
 		BindingDesc.ID = TransitionMetadata.ParametersId;
-		BindingDesc.Name = TEXT("Transition Parameters");
+		BindingDesc.Name = FName("TransitionParameters", TransitionIndex);
 		BindingDesc.Struct = Pair.Value.GetPropertyBagStruct();
 		BindingDesc.DataHandle = FSceneStateBindingDataHandle(ESceneStateDataType::Transition, TransitionIndex);
 
-		AddBindingDesc(MoveTemp(BindingDesc));
+		AddBindingDesc(MoveTemp(BindingDesc), Pair.Value.GetMutableValue());
 	}
 }
 
@@ -201,12 +201,12 @@ void FBindingCompiler::AddTaskBindingDescs()
 		{
 			FSceneStateBindingDesc BindingDesc;
 			BindingDesc.ID = TaskMetadata.TaskId;
-			BindingDesc.Name = FName(TEXT("Task"), TaskIndex);
+			BindingDesc.Name = FName("Task", TaskIndex);
 			BindingDesc.Struct = TaskInstance.GetScriptStruct();
 			BindingDesc.DataHandle = FSceneStateBindingDataHandle(ESceneStateDataType::Task, TaskIndex);
 
 			AddDataView(BindingDesc.DataHandle, TaskInstance);
-			AddBindingDesc(MoveTemp(BindingDesc));
+			AddBindingDesc(MoveTemp(BindingDesc), TaskInstance);
 		}
 
 		const FSceneStateTaskBindingExtension* BindingExtension = Task.GetBindingExtension();
@@ -230,7 +230,7 @@ void FBindingCompiler::AddTaskBindingDescs()
 				BindingDesc.DataHandle = FSceneStateBindingDataHandle(ESceneStateDataType::TaskExtension, TaskIndex, InBindingDesc.DataIndex);
 
 				This->AddDataView(BindingDesc.DataHandle, InBindingDesc.DataView);
-				This->AddBindingDesc(MoveTemp(BindingDesc));
+				This->AddBindingDesc(MoveTemp(BindingDesc), InBindingDesc.DataView);
 			});
 	}
 }
@@ -253,7 +253,7 @@ void FBindingCompiler::AddEventHandlerBindingDescs()
 		BindingDesc.Struct = EventSchema->Struct;
 		BindingDesc.DataHandle = FSceneStateBindingDataHandle(ESceneStateDataType::EventHandler, EventHandlerIndex);
 
-		AddBindingDesc(MoveTemp(BindingDesc));
+		AddBindingDesc(MoveTemp(BindingDesc), FPropertyBindingDataView(EventSchema->Struct, const_cast<uint8*>(EventSchema->Struct->GetDefaultInstance())));
 	}
 }
 
@@ -273,7 +273,7 @@ void FBindingCompiler::AddFunctionBindingDescs()
 		BindingDesc.Struct = FunctionInstance.GetScriptStruct();
 		BindingDesc.DataHandle = FSceneStateBindingDataHandle(ESceneStateDataType::Function, FunctionIndex);
 
-		AddBindingDesc(MoveTemp(BindingDesc));
+		AddBindingDesc(MoveTemp(BindingDesc), FunctionInstance);
 	}
 }
 
@@ -574,12 +574,9 @@ FSceneStateBindingDataHandle FBindingCompiler::GetDataHandleById(const FGuid& In
 	}
 
 	// State Machines
-	for (const TPair<FGuid, uint16>& Pair : TemplateData->StateMachineIdToIndex)
+	if (const uint16* StateMachineIndex = TemplateData->StateMachineIdToIndex.Find(InStructId))
 	{
-		if (Pair.Key == InStructId)
-		{
-			return FSceneStateBindingDataHandle(ESceneStateDataType::StateMachine, Pair.Value);
-		}
+		return FSceneStateBindingDataHandle(ESceneStateDataType::StateMachine, *StateMachineIndex);
 	}
 
 	// Transitions
@@ -631,12 +628,9 @@ FSceneStateBindingDataHandle FBindingCompiler::GetDataHandleById(const FGuid& In
 	}
 
 	// States
-	for (const TPair<FGuid, uint16>& Pair : TemplateData->StateParameterIdToIndex)
+	if (const uint16* StateIndex = TemplateData->StateParameterIdToIndex.Find(InStructId))
 	{
-		if (Pair.Key == InStructId)
-		{
-			return FSceneStateBindingDataHandle(ESceneStateDataType::State, Pair.Value);
-		}
+		return FSceneStateBindingDataHandle(ESceneStateDataType::State, *StateIndex);
 	}
 
 	return FSceneStateBindingDataHandle();
diff --git a/Engine/Plugins/VirtualProduction/SceneState/Source/SceneStateBlueprintEditor/Private/Compiler/SceneStateBindingCompiler.h b/Engine/Plugins/VirtualProduction/SceneState/Source/SceneStateBlueprintEditor/Private/Compiler/SceneStateBindingCompiler.h
--- a/Engine/Plugins/VirtualProduction/SceneState/Source/SceneStateBlueprintEditor/Private/Compiler/SceneStateBindingCompiler.h
+++ b/Engine/Plugins/VirtualProduction/SceneState/Source/SceneStateBlueprintEditor/Private/Compiler/SceneStateBindingCompiler.h
@@ -39,7 +39,7 @@ public:
 
 private:
 	void AddDataView(const FSceneStateBindingDataHandle& InDataHandle, FPropertyBindingDataView InDataView);
-	void AddBindingDesc(FSceneStateBindingDesc&& InBindingDesc);
+	void AddBindingDesc(FSceneStateBindingDesc&& InBindingDesc, FPropertyBindingDataView InDataView);
 
 	void AddRootBindingDesc();
 	void AddStateBindingDescs();
diff --git a/Engine/Plugins/VirtualProduction/SceneState/Source/SceneStateBlueprintEditor/Private/Compiler/SceneStateBlueprintCompilerContext.cpp b/Engine/Plugins/VirtualProduction/SceneState/Source/SceneStateBlueprintEditor/Private/Compiler/SceneStateBlueprintCompilerContext.cpp
--- a/Engine/Plugins/VirtualProduction/SceneState/Source/SceneStateBlueprintEditor/Private/Compiler/SceneStateBlueprintCompilerContext.cpp
+++ b/Engine/Plugins/VirtualProduction/SceneState/Source/SceneStateBlueprintEditor/Private/Compiler/SceneStateBlueprintCompilerContext.cpp
@@ -188,7 +188,15 @@ void FBlueprintCompilerContext::MergeUbergraphPagesIn(UEdGraph* InUbergraph)
 
 	UpdateTaskIndices();
 	SetCompiledIndices();
-	CompileBindings();
+}
+
+void FBlueprintCompilerContext::FinishCompilingClass(UClass* InClass)
+{
+	FKismetCompilerContext::FinishCompilingClass(InClass);
+	if (NewGeneratedClass->IsFullClass())
+	{
+		CompileBindings();
+	}
 }
 
 void FBlueprintCompilerContext::OnPostCDOCompiled(const UObject::FPostCDOCompiledContext& InContext)
diff --git a/Engine/Plugins/VirtualProduction/SceneState/Source/SceneStateBlueprintEditor/Private/Compiler/SceneStateBlueprintCompilerContext.h b/Engine/Plugins/VirtualProduction/SceneState/Source/SceneStateBlueprintEditor/Private/Compiler/SceneStateBlueprintCompilerContext.h
--- a/Engine/Plugins/VirtualProduction/SceneState/Source/SceneStateBlueprintEditor/Private/Compiler/SceneStateBlueprintCompilerContext.h
+++ b/Engine/Plugins/VirtualProduction/SceneState/Source/SceneStateBlueprintEditor/Private/Compiler/SceneStateBlueprintCompilerContext.h
@@ -28,6 +28,7 @@ private:
 	virtual void EnsureProperGeneratedClass(UClass*& InOutTargetClass) override;
 	virtual void SaveSubObjectsFromCleanAndSanitizeClass(FSubobjectCollection& OutSubObjectsToSave, UBlueprintGeneratedClass* InClassToClean) override;
 	virtual void MergeUbergraphPagesIn(UEdGraph* InUbergraph) override;
+	virtual void FinishCompilingClass(UClass* InClass) override;
 	virtual void OnPostCDOCompiled(const UObject::FPostCDOCompiledContext& InContext) override;
 	//~ End FKismetCompilerContext
 
