Motion Design Scene State: fixed inability to promote reference bindings.

#Synced-CL 51290016

diff --git a/Engine/Plugins/VirtualProduction/SceneState/Source/SceneStateBlueprint/Private/SceneStateBlueprintUtils.cpp b/Engine/Plugins/VirtualProduction/SceneState/Source/SceneStateBlueprint/Private/SceneStateBlueprintUtils.cpp
--- a/Engine/Plugins/VirtualProduction/SceneState/Source/SceneStateBlueprint/Private/SceneStateBlueprintUtils.cpp
+++ b/Engine/Plugins/VirtualProduction/SceneState/Source/SceneStateBlueprint/Private/SceneStateBlueprintUtils.cpp
@@ -8,6 +8,7 @@
 #include "PropertyBagDetails.h"
 #include "SceneStateBlueprint.h"
 #include "SceneStateMachineGraph.h"
+#include "SceneStateMachineGraphUtils.h"
 #include "Templates/Function.h"
 
 namespace UE::SceneState::Graph
@@ -119,8 +120,13 @@ void CreateBlueprintVariables(USceneStateBlueprint* InBlueprint, TArrayView<UE::
 		FString DefaultValue;
 		if (CreationDesc.SourceProperty && CreationDesc.SourceContainerAddress)
 		{
-			const void* SourceValue = CreationDesc.SourceProperty->ContainerPtrToValuePtr<void>(CreationDesc.SourceContainerAddress);
-			CreationDesc.SourceProperty->ExportText_Direct(DefaultValue, SourceValue, SourceValue, nullptr, PPF_None);
+			const FEdGraphPinType SourcePropertyType = GetPropertyDescAsPin(FPropertyBagPropertyDesc(CreationDesc.SourceProperty->GetFName(), CreationDesc.SourceProperty));
+			// Only copy default if the values are the exact same
+			if (SourcePropertyType == VariableType)
+			{
+				const void* SourceValue = CreationDesc.SourceProperty->ContainerPtrToValuePtr<void>(CreationDesc.SourceContainerAddress);
+				CreationDesc.SourceProperty->ExportText_Direct(DefaultValue, SourceValue, SourceValue, nullptr, PPF_None);
+			}
 		}
 
 		if (FBlueprintEditorUtils::AddMemberVariable(InBlueprint, MemberName, VariableType, DefaultValue))
diff --git a/Engine/Plugins/VirtualProduction/SceneState/Source/SceneStateBlueprintEditor/Private/Compiler/SceneStateBindingCompiler.cpp b/Engine/Plugins/VirtualProduction/SceneState/Source/SceneStateBlueprintEditor/Private/Compiler/SceneStateBindingCompiler.cpp
--- a/Engine/Plugins/VirtualProduction/SceneState/Source/SceneStateBlueprintEditor/Private/Compiler/SceneStateBindingCompiler.cpp
+++ b/Engine/Plugins/VirtualProduction/SceneState/Source/SceneStateBlueprintEditor/Private/Compiler/SceneStateBindingCompiler.cpp
@@ -98,13 +98,10 @@ void FBindingCompiler::Compile()
 	CompileReferences();
 }
 
-void FBindingCompiler::AddDataView(const FSceneStateBindingDataHandle& InDataHandle, FPropertyBindingDataView InDataView)
-{
-	TemplateDataViewMap.Add(InDataHandle, InDataView);
-}
-
 void FBindingCompiler::AddBindingDesc(FSceneStateBindingDesc&& InBindingDesc, FPropertyBindingDataView InDataView)
 {
+	TemplateDataViewMap.Add(InBindingDesc.DataHandle, InDataView);
+
 	// Record this binding in the valid binding map. This will be used to then remove all the invalid bindings from the collection
 	ValidBindingMap.Add(InBindingDesc.ID, InDataView);
 
@@ -205,7 +202,6 @@ void FBindingCompiler::AddTaskBindingDescs()
 			BindingDesc.Struct = TaskInstance.GetScriptStruct();
 			BindingDesc.DataHandle = FSceneStateBindingDataHandle(ESceneStateDataType::Task, TaskIndex);
 
-			AddDataView(BindingDesc.DataHandle, TaskInstance);
 			AddBindingDesc(MoveTemp(BindingDesc), TaskInstance);
 		}
 
@@ -229,7 +225,6 @@ void FBindingCompiler::AddTaskBindingDescs()
 				BindingDesc.Struct = InBindingDesc.DataView.GetStruct();
 				BindingDesc.DataHandle = FSceneStateBindingDataHandle(ESceneStateDataType::TaskExtension, TaskIndex, InBindingDesc.DataIndex);
 
-				This->AddDataView(BindingDesc.DataHandle, InBindingDesc.DataView);
 				This->AddBindingDesc(MoveTemp(BindingDesc), InBindingDesc.DataView);
 			});
 	}
diff --git a/Engine/Plugins/VirtualProduction/SceneState/Source/SceneStateBlueprintEditor/Private/Compiler/SceneStateBindingCompiler.h b/Engine/Plugins/VirtualProduction/SceneState/Source/SceneStateBlueprintEditor/Private/Compiler/SceneStateBindingCompiler.h
--- a/Engine/Plugins/VirtualProduction/SceneState/Source/SceneStateBlueprintEditor/Private/Compiler/SceneStateBindingCompiler.h
+++ b/Engine/Plugins/VirtualProduction/SceneState/Source/SceneStateBlueprintEditor/Private/Compiler/SceneStateBindingCompiler.h
@@ -38,7 +38,6 @@ public:
 	void Compile();
 
 private:
-	void AddDataView(const FSceneStateBindingDataHandle& InDataHandle, FPropertyBindingDataView InDataView);
 	void AddBindingDesc(FSceneStateBindingDesc&& InBindingDesc, FPropertyBindingDataView InDataView);
 
 	void AddRootBindingDesc();
diff --git a/Engine/Plugins/VirtualProduction/SceneState/Source/SceneStateBlueprintEditor/Private/DetailsView/SceneStateCachedBindingData.cpp b/Engine/Plugins/VirtualProduction/SceneState/Source/SceneStateBlueprintEditor/Private/DetailsView/SceneStateCachedBindingData.cpp
--- a/Engine/Plugins/VirtualProduction/SceneState/Source/SceneStateBlueprintEditor/Private/DetailsView/SceneStateCachedBindingData.cpp
+++ b/Engine/Plugins/VirtualProduction/SceneState/Source/SceneStateBlueprintEditor/Private/DetailsView/SceneStateCachedBindingData.cpp
@@ -337,20 +337,15 @@ void FSceneStateCachedBindingData::AddPropertyInfoOverride(const FProperty& InPr
 		{
 			TSharedRef<UE::PropertyBinding::FPropertyInfoOverride> PropertyInfoOverride = MakeShared<UE::PropertyBinding::FPropertyInfoOverride>();
 			PropertyInfoOverride->PinType = PinType;
-
-			FString TypeName;
-			if (UObject* SubCategoryObject = PinType.PinSubCategoryObject.Get())
-			{
-				TypeName = SubCategoryObject->GetName();
-			}
-			else
-			{
-				TypeName = PinType.PinCategory.ToString() + TEXT(" ") + PinType.PinSubCategory.ToString();
-			}
-
-			PropertyInfoOverride->TypeNameText = FText::FromString(TypeName);
+			PropertyInfoOverride->TypeNameText = UEdGraphSchema_K2::TypeToText(PinType);
 			OutPropertyInfoOverrides.Emplace(MoveTemp(PropertyInfoOverride));
 		}
+
+		TSharedRef<UE::PropertyBinding::FPropertyInfoOverride> PropertyInfoOverride = MakeShared<UE::PropertyBinding::FPropertyInfoOverride>();
+		PropertyInfoOverride->PinType.PinCategory = UEdGraphSchema_K2::PC_Struct;
+		PropertyInfoOverride->PinType.PinSubCategoryObject = StructProperty->Struct;
+		PropertyInfoOverride->TypeNameText = LOCTEXT("PromotePropertyReferenceDisplayName", "Promote as Reference");
+		OutPropertyInfoOverrides.Emplace(MoveTemp(PropertyInfoOverride));
 	}
 }
 
