Motion Design Scene State: fixed issues with packaged builds not calling template data link, causing transition and conduit evaluation events to not be set and never called.

#Synced-CL 51764840

diff --git a/Engine/Plugins/VirtualProduction/SceneState/Source/SceneState/Private/SceneStateGeneratedClass.cpp b/Engine/Plugins/VirtualProduction/SceneState/Source/SceneState/Private/SceneStateGeneratedClass.cpp
--- a/Engine/Plugins/VirtualProduction/SceneState/Source/SceneState/Private/SceneStateGeneratedClass.cpp
+++ b/Engine/Plugins/VirtualProduction/SceneState/Source/SceneState/Private/SceneStateGeneratedClass.cpp
@@ -22,11 +22,7 @@ USceneStateGeneratedClass::USceneStateGeneratedClass()
 void USceneStateGeneratedClass::Link(FArchive& Ar, bool bInRelinkExistingProperties)
 {
 	Super::Link(Ar, bInRelinkExistingProperties);
-
-	if (TemplateData && IsFullClass())
-	{
-		TemplateData->Link(this);
-	}
+	LinkData();
 }
 
 void USceneStateGeneratedClass::PurgeClass(bool bInRecompilingOnLoad)
@@ -41,6 +37,7 @@ void USceneStateGeneratedClass::PurgeClass(bool bInRecompilingOnLoad)
 void USceneStateGeneratedClass::PostLoad()
 {
 	Super::PostLoad();
+	LinkData();
 	ResolveBindings();
 }
 
@@ -49,6 +46,14 @@ void USceneStateGeneratedClass::BeginDestroy()
 	Super::BeginDestroy();
 }
 
+void USceneStateGeneratedClass::LinkData()
+{
+	if (TemplateData && IsFullClass())
+	{
+		TemplateData->Link(this);
+	}
+}
+
 void USceneStateGeneratedClass::ResolveBindings()
 {
 	if (TemplateData && IsFullClass())
diff --git a/Engine/Plugins/VirtualProduction/SceneState/Source/SceneState/Public/SceneStateGeneratedClass.h b/Engine/Plugins/VirtualProduction/SceneState/Source/SceneState/Public/SceneStateGeneratedClass.h
--- a/Engine/Plugins/VirtualProduction/SceneState/Source/SceneState/Public/SceneStateGeneratedClass.h
+++ b/Engine/Plugins/VirtualProduction/SceneState/Source/SceneState/Public/SceneStateGeneratedClass.h
@@ -48,6 +48,9 @@ public:
 	//~ End UObject
 
 private:
+	/** Called to link the template data */
+	void LinkData();
+
 	/** Called to resolve the template data bindings */
 	UE_INTERNAL UE_API void ResolveBindings();
 
